- Updating MediaIO to work with TimedData monitor

diff --git a/Engine/Plugins/VirtualProduction/TimedDataMonitor/Source/TimedDataMonitor/Private/TimedDataMonitorSubsystem.cpp b/Engine/Plugins/VirtualProduction/TimedDataMonitor/Source/TimedDataMonitor/Private/TimedDataMonitorSubsystem.cpp
--- a/Engine/Plugins/VirtualProduction/TimedDataMonitor/Source/TimedDataMonitor/Private/TimedDataMonitorSubsystem.cpp
+++ b/Engine/Plugins/VirtualProduction/TimedDataMonitor/Source/TimedDataMonitor/Private/TimedDataMonitorSubsystem.cpp
@@ -3,13 +3,15 @@
 #include "TimedDataMonitorSubsystem.h"
 
 #include "Engine/Engine.h"
-#include "Engine/TimecodeProvider.h"
-#include "HAL/PlatformTime.h"
+#include "Engine/World.h"
 #include "ITimeManagementModule.h"
+#include "LatentActions.h"
 #include "Misc/App.h"
 #include "Misc/QualifiedFrameTime.h"
 #include "Stats/Stats2.h"
 #include "TimedDataInputCollection.h"
+#include "TimedDataMonitorCalibration.h"
+#include "UObject/Stack.h"
 
 
 static TAutoConsoleVariable<bool> CVarEnableTimedDataMonitorSubsystemStats(TEXT("TimedDataMonitor.EnableStatUpdate"), 1, TEXT("Enable calculating evaluation statistics of all registered channels."));
@@ -108,6 +110,28 @@ ITimedDataInputChannel* UTimedDataMonitorSubsystem::GetTimedDataChannel(const FT
 }
 
 
+ double UTimedDataMonitorSubsystem::GetEvaluationTime(ETimedDataInputEvaluationType EvaluationType)
+ {
+	 double Result = 0.0;
+	 switch (EvaluationType)
+	 {
+	 case ETimedDataInputEvaluationType::Timecode:
+		 if (FApp::GetCurrentFrameTime().IsSet())
+		 {
+			 Result = FApp::GetCurrentFrameTime().GetValue().AsSeconds();
+		 }
+		 break;
+	 case ETimedDataInputEvaluationType::PlatformTime:
+		 Result = FApp::GetCurrentTime();
+		 break;
+	 case ETimedDataInputEvaluationType::None:
+	 default:
+		 break;
+	 }
+	 return Result;
+ }
+
+
 TArray<FTimedDataMonitorInputIdentifier> UTimedDataMonitorSubsystem::GetAllInputs()
 {
 	BuildSourcesListIfNeeded();
@@ -128,470 +152,88 @@ TArray<FTimedDataMonitorChannelIdentifier> UTimedDataMonitorSubsystem::GetAllCha
 }
 
 
-FTimedDataMonitorCalibrationResult UTimedDataMonitorSubsystem::CalibrateWithTimecodeProvider()
+TArray<FTimedDataMonitorChannelIdentifier> UTimedDataMonitorSubsystem::GetAllEnabledChannels()
 {
 	BuildSourcesListIfNeeded();
 
-	FTimedDataMonitorCalibrationResult Result;
-	Result.ReturnCode = ETimedDataMonitorCalibrationReturnCode::Failed_CanNotCallibrateWithoutJam;
-//	return Result;
-//
-//
-//	UTimecodeProvider* CurrentTimecodeProvider = GEngine->GetTimecodeProvider();
-//	if (CurrentTimecodeProvider == nullptr
-//		|| CurrentTimecodeProvider->GetSynchronizationState() != ETimecodeProviderSynchronizationState::Synchronized
-//		|| !FApp::GetCurrentFrameTime().IsSet())
-//	{
-//		Result.ReturnCode = ETimedDataMonitorCallibrationReturnCode::Failed_NoTimecode;
-//		return Result;
-//	}
-//	FQualifiedFrameTime CurrentFrameTime = CurrentTimecodeProvider->GetQualifiedFrameTime();
-//
-//	struct FCollectedDataTimes
-//	{
-//		FTimedDataMonitorInputIdentifier InputId;
-//		TArray<FTimedDataInputSampleTime> DataTimes;
-//	};
-//
-//	// Collect all DataTimes
-//	TArray<FCollectedDataTimes> AllCollectedDataTimes;
-//	AllCollectedDataTimes.Reset(InputMap.Num());
-//	for (const auto& InputItt : InputMap)
-//	{
-//		if (InputItt.Value.bEnabled)
-//		{
-//			if (InputItt.Value.Input->GetState() != ETimedDataInputState::Connected)
-//			{
-//				Result.ReturnCode = ETimedDataMonitorCallibrationReturnCode::Failed_UnresponsiveInput;
-//				Result.FailureInputIdentifiers.Add(InputItt.Key);
-//				return Result;
-//			}
-//
-//			if (InputItt.Value.Input->GetEvaluationType() != ETimedDataInputEvaluationType::Timecode)
-//			{
-//				Result.ReturnCode = ETimedDataMonitorCallibrationReturnCode::Failed_NoTimecode;
-//				Result.FailureInputIdentifiers.Add(InputItt.Key);
-//				return Result;
-//			}
-//
-//			FFrameRate FrameRate = InputItt.Value.Input->GetFrameRate();
-//			if (FrameRate == ITimedDataInput::UnknowedFrameRate)
-//			{
-//				Result.ReturnCode = ETimedDataMonitorCallibrationReturnCode::Failed_InvalidFrameRate;
-//				Result.FailureInputIdentifiers.Add(InputItt.Key);
-//				return Result;
-//			}
-//
-//			TArray<FTimedDataInputSampleTime> DataTimes = InputItt.Value.Input->GetDataTimes();
-//			if (DataTimes.Num() == 0)
-//			{
-//				Result.ReturnCode = ETimedDataMonitorCallibrationReturnCode::Failed_NoDataBuffered;
-//				Result.FailureInputIdentifiers.Add(InputItt.Key);
-//				return Result;
-//			}
-//
-//			FCollectedDataTimes Element;
-//			Element.InputId = InputItt.Key;
-//			Element.DataTimes = MoveTemp(DataTimes);
-//			AllCollectedDataTimes.Emplace(MoveTemp(Element));
-//		}
-//	}
-//
-//	if (AllCollectedDataTimes.Num() == 0)
-//	{
-//		Result.ReturnCode = ETimedDataMonitorCallibrationReturnCode::Succeeded;
-//		return Result;
-//	}
-//
-//	// Is there a range of data that everyone is happy with
-//	{
-//		// With [A-C] it will return [10-11], with [A-D] it should not be able to find anything
-//		//TC == 12
-//		//A    10  11  12  13
-//		//B 9  10  11
-//		//C    10  11  12
-//		//D            12  13
-//
-//		FFrameTime RangeMin;
-//		FFrameTime RangeMax;
-//		bool bFirstItem = true;
-//		bool bFoundRange = true;
-//
-//		for (const FCollectedDataTimes& DataTimesItt : AllCollectedDataTimes)
-//		{
-//			check(DataTimesItt.DataTimes.Num() > 0);
-//
-//			// On purpose do not use the evaluation offset
-//			//const FFrameTime FirstFrameTime = DataTimesItt.DataTimes[0].Timecode.ConvertTo(CurrentFrameTime.Rate) - DataTimesItt.EvaluationOffset.ConvertTo(CurrentFrameTime.Rate);
-//
-//			const FFrameTime FirstFrameTime = DataTimesItt.DataTimes[0].Timecode.ConvertTo(CurrentFrameTime.Rate);
-//			const FFrameTime LastFrameTime = DataTimesItt.DataTimes.Last().Timecode.ConvertTo(CurrentFrameTime.Rate);
-//
-//			if (!bFirstItem)
-//			{
-//				if (FirstFrameTime <= RangeMax && LastFrameTime >= RangeMin)
-//				{
-//
-//					RangeMin = FMath::Max(RangeMin, FirstFrameTime);
-//					RangeMax = FMath::Min(RangeMax, LastFrameTime);
-//				}
-//				else
-//				{
-//					// Return an unset value
-//					bFoundRange = false;
-//					break;
-//				}
-//			}
-//			else
-//			{
-//				bFirstItem = false;
-//				RangeMin = FirstFrameTime;
-//				RangeMax = LastFrameTime;
-//			}
-//		}
-//
-//		if (bFoundRange)
-//		{
-//			FFrameTime TimecodeProviderOffset = RangeMax - CurrentFrameTime.Time;
-//			if (RangeMin <= CurrentFrameTime.Time && CurrentFrameTime.Time <= RangeMax)
-//			{
-//				// TC in the range, if so use the TC provider value
-//
-//				// Reset the FrameDelay
-//				check(CurrentTimecodeProvider);
-//				CurrentTimecodeProvider->FrameDelay = 0.f;
-//
-//				// Reset previous evaluation offset
-//				for (const FCollectedDataTimes& DataTimesItt : AllCollectedDataTimes)
-//				{
-//					InputMap[DataTimesItt.InputId].Input->SetEvaluationOffsetInSeconds(0.0);
-//				}
-//				Result.ReturnCode = ETimedDataMonitorCallibrationReturnCode::Succeeded;
-//				return Result;
-//			}
-//			else if (CurrentFrameTime.Time < RangeMin)
-//			{
-//				// We need to increase buffer size of all inputs if possible
-////@TODO
-//				//Is it possible
-//				//else set TimecodeProviderOffset to offset the TC provider
-//
-//				//if (UTimecodeProvider* TimecodeProvider = GEngine->GetTimecodeProvider())
-//				//{
-//				//	TimecodeProvider->FrameDelay = 0.f;
-//				//}
-//				//Result.FailureInputIdentifiers.Reset(AllCollectedDataTimes.Num());
-//				//for (const FCollectedDataTimes& DataTimesItt : AllCollectedDataTimes)
-//				//{
-//				//	InputMap[DataTimesItt.InputId].Input->SetEvaluationOffsetInSeconds(0.0);
-//				//	InputMap[DataTimesItt.InputId].Input->SetDataBufferSize(0);
-//				//}
-//				//Result.ReturnCode = ETimedDataMonitorCallibrationReturnCode::Retry_BufferSizeHasBeenIncreased;
-//				//return Result;
-//			}
-//
-//
-//			// When we can't resize the buffer (because they would be too big, then offset the TC provider.
-//			//Or when the TC Provider is over all the input, then offset the TC provider
-//			check(CurrentTimecodeProvider);
-//			CurrentTimecodeProvider->FrameDelay = TimecodeProviderOffset.AsDecimal();
-//	
-//			// Reset previous evaluation offset
-//			for (const FCollectedDataTimes& DataTimesItt : AllCollectedDataTimes)
-//			{
-//				InputMap[DataTimesItt.InputId].Input->SetEvaluationOffsetInSeconds(0.0);
-//			}
-//			Result.ReturnCode = ETimedDataMonitorCallibrationReturnCode::Succeeded;
-//			return Result;
-//		}
-//	}
-//
-//	// No range was found. What is the smallest acceptable LastFrameTime. Can we increase a buffer?
-//	{
-//		// It will return 11, and we will try to increase the buffer size of E by 1 and jam (A, F, G). We want to be closer to TC as possible.
-//		// TC == 50.
-//		//A 1  2
-//		//B                  48  49  50  51
-//		//C              47  48  49
-//		//D                  48  49  50
-//		//E                          50  51
-//		//F 1
-//		//G                                                     100
-//
-//
-//		//@todo. complete the algo
-//	}
-//
-//	Result.ReturnCode = ETimedDataMonitorCallibrationReturnCode::Failed_CanNotCallibrateWithoutJam;
+	TArray<FTimedDataMonitorChannelIdentifier> Result;
+	Result.Reset(ChannelMap.Num());
+	for (const auto& ChannelMapItt : ChannelMap)
+	{
+		if (ChannelMapItt.Value.bEnabled)
+		{
+			Result.Add(ChannelMapItt.Key);
+		}
+	}
 	return Result;
 }
 
 
-FTimedDataMonitorJamResult UTimedDataMonitorSubsystem::JamInputs(ETimedDataInputEvaluationType InEvaluationType)
+void UTimedDataMonitorSubsystem::CalibrateLatent(UObject* WorldContextObject, struct FLatentActionInfo LatentInfo, const FTimedDataMonitorCalibrationParameters& CalibrationParameters, FTimedDataMonitorCalibrationResult& Result)
 {
-	BuildSourcesListIfNeeded();
-
-	FTimedDataMonitorJamResult Result;
-	Result.ReturnCode = ETimedDataMonitorJamReturnCode::Succeeded;
-
-
-	const double CurrentPlatformTime = FApp::GetCurrentTime();
-	UTimecodeProvider* CurrentTimecodeProvider = GEngine->GetTimecodeProvider();
-	FQualifiedFrameTime CurrentFrameTime;
-	if (InEvaluationType == ETimedDataInputEvaluationType::Timecode)
+	struct FCalibrateAction : public FPendingLatentAction
 	{
-		if (CurrentTimecodeProvider == nullptr || CurrentTimecodeProvider->GetSynchronizationState() != ETimecodeProviderSynchronizationState::Synchronized || !FApp::GetCurrentFrameTime().IsSet())
+	public:
+		FName ExecutionFunction;
+		int32 Linkage;
+		FWeakObjectPtr CallbackTarget;
+		FTimedDataMonitorCalibrationResult& Result;
+		TUniquePtr<FTimedDataMonitorCalibration> Calibration;
+		bool bOnCompleted;
+
+		FCalibrateAction(const FLatentActionInfo& InLatentInfo, FTimedDataMonitorCalibrationResult& InResult)
+			: FPendingLatentAction()
+			, ExecutionFunction(InLatentInfo.ExecutionFunction)
+			, Linkage(InLatentInfo.Linkage)
+			, CallbackTarget(InLatentInfo.CallbackTarget)
+			, Result(InResult)
+			, Calibration(new FTimedDataMonitorCalibration)
+			, bOnCompleted(false)
 		{
-			Result.ReturnCode = ETimedDataMonitorJamReturnCode::Failed_NoTimecode;
-			return Result;
 		}
 
-		CurrentFrameTime = CurrentTimecodeProvider->GetQualifiedFrameTime();
-	}
-
-
-	// Test all Channels for Failed_UnresponsiveInput
-	TArray<FTimedDataMonitorInputIdentifier> AllValidInputIndentifiers;
-	for (const auto& ChannelItt : ChannelMap)
-	{
-		if (ChannelItt.Value.bEnabled)
-		{	
-			if (ChannelItt.Value.Channel->GetState() != ETimedDataInputState::Connected)
+		virtual void UpdateOperation(FLatentResponse& Response) override
+		{
+			if (bOnCompleted)
 			{
-				Result.ReturnCode = ETimedDataMonitorJamReturnCode::Failed_UnresponsiveInput;
-				Result.FailureChannelIdentifiers.Add(ChannelItt.Key);
+				Response.FinishAndTriggerIf(true, ExecutionFunction, Linkage, CallbackTarget);
 			}
-
-			AllValidInputIndentifiers.AddUnique(ChannelItt.Value.InputIdentifier);
 		}
-	}
-
-	if (Result.ReturnCode == ETimedDataMonitorJamReturnCode::Failed_UnresponsiveInput)
-	{
-		return Result;
-	}
-
-	// Test all inputs for Failed_EvaluationTypeDoNotMatch
-	for (const FTimedDataMonitorInputIdentifier& InputIndentifier : AllValidInputIndentifiers)
-	{
-		// The UI show have make sure that all evaluation type are correct
-		if (InputMap[InputIndentifier].Input->GetEvaluationType() != InEvaluationType)
-		{
-			Result.ReturnCode = ETimedDataMonitorJamReturnCode::Failed_EvaluationTypeDoNotMatch;
-			Result.FailureInputIdentifiers.Add(InputIndentifier);
-		}
-	}
-
-	if (Result.ReturnCode == ETimedDataMonitorJamReturnCode::Failed_EvaluationTypeDoNotMatch)
-	{
-		return Result;
-	}
 
-	if (InEvaluationType == ETimedDataInputEvaluationType::None)
-	{
-		// Set the evaluation offset of everyone to 0
-		for (const FTimedDataMonitorInputIdentifier& InputId : AllValidInputIndentifiers)
+#if WITH_EDITOR
+		virtual FString GetDescription() const override
 		{
-			InputMap[InputId].Input->SetEvaluationOffsetInSeconds(0.0);
+			return FString::Printf(TEXT("Calibrating."));
 		}
-		Result.ReturnCode = ETimedDataMonitorJamReturnCode::Succeeded;
-	}
-	else
-	{
-		// Collect all DataTimes
+#endif
 
-		struct FChannelMinMaxSampleTime
-		{
-			FChannelMinMaxSampleTime(FTimedDataChannelSampleTime InMin, FTimedDataChannelSampleTime InMax)
-				: Min(InMin), Max(InMax) {}
-			FTimedDataChannelSampleTime Min;
-			FTimedDataChannelSampleTime Max;
-		};
-
-		TMap<FTimedDataMonitorChannelIdentifier, FChannelMinMaxSampleTime> AllCollectedDataTimes;
-		for (const auto& ChannelItt : ChannelMap)
+		void OnCompleted(FTimedDataMonitorCalibrationResult InResult)
 		{
-			if (ChannelItt.Value.bEnabled)
-			{
-				if (ChannelItt.Value.Channel->GetNumberOfSamples() <= 0)
-				{
-					Result.ReturnCode = ETimedDataMonitorJamReturnCode::Failed_NoDataBuffered;
-					Result.FailureChannelIdentifiers.Add(ChannelItt.Key);
-				}
-
-				FChannelMinMaxSampleTime DataTimes = FChannelMinMaxSampleTime(ChannelItt.Value.Channel->GetOldestDataTime() , ChannelItt.Value.Channel->GetNewestDataTime());
-				AllCollectedDataTimes.Add(ChannelItt.Key, MoveTemp(DataTimes));
-			}
+			Result = InResult;
+			bOnCompleted = true;
 		}
+	};
 
-		if (Result.ReturnCode == ETimedDataMonitorJamReturnCode::Failed_NoDataBuffered)
+	if (UWorld* World = GEngine->GetWorldFromContextObject(WorldContextObject, EGetWorldErrorMode::LogAndReturnNull))
+	{
+		FLatentActionManager& LatentManager = World->GetLatentActionManager();
+		if (LatentManager.FindExistingAction<FCalibrateAction>(LatentInfo.CallbackTarget, LatentInfo.UUID) == nullptr)
 		{
-			return Result;
+			FCalibrateAction* NewAction = new FCalibrateAction(LatentInfo, Result);
+			NewAction->Calibration->CalibrateWithTimecode(CalibrationParameters, FTimedDataMonitorCalibration::FOnCalibrationCompletedSignature::CreateRaw(NewAction, &FCalibrateAction::OnCompleted));
+			LatentManager.AddNewAction(LatentInfo.CallbackTarget, LatentInfo.UUID, NewAction);
 		}
-
-		// Grouped by input, find what could we do to match data.
-		// EvaluationTime == 50.
-		//A1 10  11
-		//A2                   48  49  50  51
-		//A3     11  12
-		//B1                                                 99  100
-		//B2                                                     100
-		//B3                                                     100  101
-		//C1 10  11
-		//C2                                                     100
-		//D1                       49  50  51
-		//D2                   48  49  50  51
-		// For InputA, we should increase the buffer size of A2 and set an offset so that 11 == 50
-		// For InputB, we set an offset so that 100 == 50
-		// For InputC, we cannot find anything since the difference is too big. Failed.
-		// For InputD, we set an offset so that 50 == 50
-
-		// For each input, go though all their channel and find what section that matches
-		bool bAllInputInRange = true;
-		const double EvaluationTime = InEvaluationType == ETimedDataInputEvaluationType::Timecode ? CurrentFrameTime.AsSeconds() : CurrentPlatformTime;
-		for (const FTimedDataMonitorInputIdentifier& InputId : AllValidInputIndentifiers)
+		else
 		{
-			// Find newest PlatformTime
-			double SmallestMinInSeconds = TNumericLimits<double>::Max();
-			double BiggestMaxInSeconds = TNumericLimits<double>::Min();
-			double BiggerMinInSeconds = TNumericLimits<double>::Min();
-			double SmallestMaxInSeconds = TNumericLimits<double>::Max();
-			for (const FTimedDataMonitorChannelIdentifier& ChannelId : InputMap[InputId].ChannelIdentifiers)
-			{
-				// test if the sample was collect (enabled)
-				if (FChannelMinMaxSampleTime* SamplesTimes = AllCollectedDataTimes.Find(ChannelId))
-				{
-					SmallestMinInSeconds = FMath::Min(GetSeconds(InEvaluationType, SamplesTimes->Min), SmallestMinInSeconds);	//A == 10, B == 99, C == 10, D == 48
-					BiggestMaxInSeconds = FMath::Max(GetSeconds(InEvaluationType, SamplesTimes->Max), BiggestMaxInSeconds);		//A == 51, B == 101, C == 100, D == 51
-
-					BiggerMinInSeconds = FMath::Max(GetSeconds(InEvaluationType, SamplesTimes->Min), BiggerMinInSeconds);		//A == 48, B == 100, C == 10, D == 49
-					SmallestMaxInSeconds = FMath::Min(GetSeconds(InEvaluationType, SamplesTimes->Max), SmallestMaxInSeconds);	//A == 11, B == 100, C == 11, D == 51
-				}
-			}
-
-		//@todo use the stat when we are confident that they works properly
-			//const double DistanceToNewestSTD = GetInputEvaluationDistanceToNewestSampleStandardDeviation(InputId) * NumberOfSigmaOfSignification;
-			const double ExtraBufferWhenJamming = 0.5;
-
-			// Test if all the samples are in the range of the EvaluationTime 
-			bool bAllChannelInRangeOfEvaluationTime = true;
-			bool bAllChannelInRangeOfSmallestMax = true;
-			for (const FTimedDataMonitorChannelIdentifier& ChannelId : InputMap[InputId].ChannelIdentifiers)
-			{
-				if (FChannelMinMaxSampleTime* SamplesTimes = AllCollectedDataTimes.Find(ChannelId))
-				{
-					if (GetSeconds(InEvaluationType, SamplesTimes->Min) > EvaluationTime - ExtraBufferWhenJamming || GetSeconds(InEvaluationType, SamplesTimes->Max) < EvaluationTime - ExtraBufferWhenJamming)
-					{
-						bAllChannelInRangeOfEvaluationTime = false;
-					}
-
-					if (GetSeconds(InEvaluationType, SamplesTimes->Min) > SmallestMaxInSeconds - ExtraBufferWhenJamming || GetSeconds(InEvaluationType, SamplesTimes->Max) < SmallestMaxInSeconds - ExtraBufferWhenJamming)
-					{
-						bAllChannelInRangeOfSmallestMax = false;
-					}
-				}
-			}
-
-			if (bAllChannelInRangeOfEvaluationTime)
-			{
-				// Set the evaluation offset for later (case D)
-				Result.ReturnCode = ETimedDataMonitorJamReturnCode::Succeeded;
-
-				InputMap[InputId].Input->SetEvaluationOffsetInSeconds(ExtraBufferWhenJamming);
-	//@todo reset the buffer size base on the stat
-			}
-			else if (bAllChannelInRangeOfSmallestMax)
-			{
-				// Set the evaluation offset for later (case B)
-				Result.ReturnCode = ETimedDataMonitorJamReturnCode::Succeeded;
-
-				InputMap[InputId].Input->SetEvaluationOffsetInSeconds(EvaluationTime - SmallestMaxInSeconds + ExtraBufferWhenJamming);
-	//@todo reset the buffer size base on the stat
-			}
-			else
-			{
-				// Test to see if we can increment the buffer size (case A or C)
-				Result.ReturnCode = ETimedDataMonitorJamReturnCode::Retry_BufferSizeHasBeenIncreased;
-
-				if (InputMap[InputId].Input->IsDataBufferSizeControlledByInput())
-				{
-					// Get the average in delta time of the last 10 frames
-					double AverageBetweenSample = 0.0;
-					int32 AverageCounter = 0;
-					for (const FTimedDataMonitorChannelIdentifier& ChannelId : InputMap[InputId].ChannelIdentifiers)
-					{
-						if (ChannelMap[ChannelId].bEnabled)
-						{
-							TArray<FTimedDataChannelSampleTime> AllSamplesTimes = ChannelMap[ChannelId].Channel->GetDataTimes();
-							if (AllSamplesTimes.Num() > 1)
-							{
-								const double CurrentAverageBetweenSample = CalculateAverageInDeltaTimeBetweenSample(InEvaluationType, AllSamplesTimes);
-
-								++AverageCounter;
-								AverageBetweenSample += (CurrentAverageBetweenSample - AverageBetweenSample) / (double)AverageCounter;
-							}
-						}
-					}
-					if (FMath::IsNearlyZero(AverageBetweenSample))
-					{
-						AverageBetweenSample = FApp::GetDeltaTime();
-					}
-
-					const int32 TotalNumberOfFrames = (BiggestMaxInSeconds - SmallestMaxInSeconds - ExtraBufferWhenJamming) / AverageBetweenSample;
-					const int32 CurrentDataBufferSize = InputMap[InputId].Input->GetDataBufferSize();
-					InputMap[InputId].Input->SetDataBufferSize(TotalNumberOfFrames);
-					const int32 UpdatedDataBufferSize = InputMap[InputId].Input->GetDataBufferSize();
-					if (UpdatedDataBufferSize < TotalNumberOfFrames)
-					{
-						// We were not able to increase the buffer size (case C) 
-						Result.ReturnCode = ETimedDataMonitorJamReturnCode::Failed_BufferSizeHaveBeenMaxed;
-						Result.FailureInputIdentifiers.Add(InputId);
-					}
-				}
-				else
-				{
-					// For each channel, check if we need to increase the buffer size. If so, by how much
-					for (const FTimedDataMonitorChannelIdentifier& ChannelId : InputMap[InputId].ChannelIdentifiers)
-					{
-						if (FChannelMinMaxSampleTime* SamplesTimes = AllCollectedDataTimes.Find(ChannelId))
-						{
-							if (GetSeconds(InEvaluationType, SamplesTimes->Min) > SmallestMaxInSeconds - ExtraBufferWhenJamming || GetSeconds(InEvaluationType, SamplesTimes->Max) < SmallestMaxInSeconds - ExtraBufferWhenJamming)
-							{
-								TArray<FTimedDataChannelSampleTime> AllSamplesTimes = ChannelMap[ChannelId].Channel->GetDataTimes();
-								const double AverageBetweenSample = CalculateAverageInDeltaTimeBetweenSample(InEvaluationType, AllSamplesTimes);
-								const int32 NumberOfNewFrameRequested = (GetSeconds(InEvaluationType, SamplesTimes->Min) - SmallestMaxInSeconds - ExtraBufferWhenJamming) / AverageBetweenSample;
-
-								const int32 CurrentDataBufferSize = ChannelMap[ChannelId].Channel->GetDataBufferSize();
-								const int32 RequestedBufferSize = NumberOfNewFrameRequested + CurrentDataBufferSize;
-								ChannelMap[ChannelId].Channel->SetDataBufferSize(RequestedBufferSize);
-								const int32 UpdatedDataBufferSize = ChannelMap[ChannelId].Channel->GetDataBufferSize();
-								if (UpdatedDataBufferSize < RequestedBufferSize)
-								{
-									// We were not able to increase the buffer size (case C) 
-									Result.ReturnCode = ETimedDataMonitorJamReturnCode::Failed_BufferSizeHaveBeenMaxed;
-									Result.FailureChannelIdentifiers.Add(ChannelId);
-								}
-							}
-						}
-					}
-				}
-
-				InputMap[InputId].Input->SetEvaluationOffsetInSeconds(EvaluationTime - SmallestMaxInSeconds + ExtraBufferWhenJamming);
-
-				if (Result.ReturnCode == ETimedDataMonitorJamReturnCode::Retry_BufferSizeHasBeenIncreased)
-				{
-					return Result;
-				}
-			}
-
-			if (InEvaluationType == ETimedDataInputEvaluationType::Timecode && CurrentTimecodeProvider)
-			{
-				CurrentTimecodeProvider->FrameDelay = 0.f;
-			}
+			FFrame::KismetExecutionMessage(TEXT("The calibration is already running."), ELogVerbosity::Warning, "CalibrationActionAlreadyStarted");
 		}
 	}
+};
 
-	return Result;
+
+FTimedDataMonitorTimeCorrectionResult UTimedDataMonitorSubsystem::ApplyTimeCorrection(const FTimedDataMonitorInputIdentifier& InputIdentifier, const FTimedDataMonitorTimeCorrectionParameters& TimeCorrectionParameters)
+{
+	return FTimedDataMonitorCalibration::ApplyTimeCorrection(InputIdentifier, TimeCorrectionParameters);
 }
 
 
@@ -607,6 +249,26 @@ void UTimedDataMonitorSubsystem::ResetAllBufferStats()
 }
 
 
+ETimedDataMonitorEvaluationState UTimedDataMonitorSubsystem::GetEvaluationState()
+{
+	BuildSourcesListIfNeeded();
+
+	ETimedDataMonitorEvaluationState WorstState = ETimedDataMonitorEvaluationState::NoSample;
+	if (InputMap.Num() > 0)
+	{
+		WorstState = ETimedDataMonitorEvaluationState::Disabled;
+		for (const auto& InputItt : InputMap)
+		{
+			const ETimedDataMonitorEvaluationState InputState = GetInputEvaluationState(InputItt.Key);
+			uint8 InputValue = FMath::Min((uint8)InputState, (uint8)WorstState);
+			WorstState = (ETimedDataMonitorEvaluationState)InputValue;
+		}
+	}
+
+	return WorstState;
+}
+
+
 bool UTimedDataMonitorSubsystem::DoesInputExist(const FTimedDataMonitorInputIdentifier& Identifier)
 {
 	BuildSourcesListIfNeeded();
@@ -619,7 +281,7 @@ ETimedDataMonitorInputEnabled UTimedDataMonitorSubsystem::GetInputEnabled(const
 {
 	BuildSourcesListIfNeeded();
 
-	if (FTimeDataInputItem* InputItem = InputMap.Find(Identifier))
+	if (const FTimeDataInputItem* InputItem = InputMap.Find(Identifier))
 	{
 		int32 bCountEnabled = 0;
 		int32 bCountDisabled = 0;
@@ -653,7 +315,7 @@ void UTimedDataMonitorSubsystem::SetInputEnabled(const FTimedDataMonitorInputIde
 {
 	BuildSourcesListIfNeeded();
 
-	if (FTimeDataInputItem* InputItem = InputMap.Find(Identifier))
+	if (const FTimeDataInputItem* InputItem = InputMap.Find(Identifier))
 	{
 		for (const FTimedDataMonitorChannelIdentifier& ChannelId : InputItem->ChannelIdentifiers)
 		{
@@ -667,7 +329,7 @@ FText UTimedDataMonitorSubsystem::GetInputDisplayName(const FTimedDataMonitorInp
 {
 	BuildSourcesListIfNeeded();
 
-	if (FTimeDataInputItem* SourceItem = InputMap.Find(Identifier))
+	if (const FTimeDataInputItem* SourceItem = InputMap.Find(Identifier))
 	{
 		return SourceItem->Input->GetDisplayName();
 	}
@@ -680,7 +342,7 @@ TArray<FTimedDataMonitorChannelIdentifier> UTimedDataMonitorSubsystem::GetInputC
 {
 	BuildSourcesListIfNeeded();
 
-	if (FTimeDataInputItem* SourceItem = InputMap.Find(Identifier))
+	if (const FTimeDataInputItem* SourceItem = InputMap.Find(Identifier))
 	{
 		return SourceItem->ChannelIdentifiers;
 	}
@@ -693,7 +355,7 @@ ETimedDataInputEvaluationType UTimedDataMonitorSubsystem::GetInputEvaluationType
 {
 	BuildSourcesListIfNeeded();
 
-	if (FTimeDataInputItem* SourceItem = InputMap.Find(Identifier))
+	if (const FTimeDataInputItem* SourceItem = InputMap.Find(Identifier))
 	{
 		return SourceItem->Input->GetEvaluationType();
 	}
@@ -717,7 +379,7 @@ float UTimedDataMonitorSubsystem::GetInputEvaluationOffsetInSeconds(const FTimed
 {
 	BuildSourcesListIfNeeded();
 
-	if (FTimeDataInputItem* SourceItem = InputMap.Find(Identifier))
+	if (const FTimeDataInputItem* SourceItem = InputMap.Find(Identifier))
 	{
 		return (float)SourceItem->Input->GetEvaluationOffsetInSeconds();
 	}
@@ -741,7 +403,7 @@ FFrameRate UTimedDataMonitorSubsystem::GetInputFrameRate(const FTimedDataMonitor
 {
 	BuildSourcesListIfNeeded();
 
-	if (FTimeDataInputItem* SourceItem = InputMap.Find(Identifier))
+	if (const FTimeDataInputItem* SourceItem = InputMap.Find(Identifier))
 	{
 		return SourceItem->Input->GetFrameRate();
 	}
@@ -754,16 +416,25 @@ FTimedDataChannelSampleTime UTimedDataMonitorSubsystem::GetInputOldestDataTime(c
 {
 	BuildSourcesListIfNeeded();
 
-	FTimedDataChannelSampleTime ResultSampleTime(TNumericLimits<double>::Max(), FQualifiedFrameTime());
-	if (FTimeDataInputItem* SourceItem = InputMap.Find(Identifier))
+	FTimedDataChannelSampleTime ResultSampleTime(0.0, FQualifiedFrameTime());
+	if (const FTimeDataInputItem* SourceItem = InputMap.Find(Identifier))
 	{
-		for (FTimedDataMonitorChannelIdentifier ChannelIdentifier : SourceItem->ChannelIdentifiers)
+		bool bFirstElement = true;
+		for (const FTimedDataMonitorChannelIdentifier ChannelIdentifier : SourceItem->ChannelIdentifiers)
 		{
 			FTimedDataChannelSampleTime OldestSampleTime = ChannelMap[ChannelIdentifier].Channel->GetOldestDataTime();
-			ResultSampleTime.PlatformSecond = FMath::Min(OldestSampleTime.PlatformSecond, ResultSampleTime.PlatformSecond);
-			if (OldestSampleTime.Timecode.AsSeconds() < ResultSampleTime.Timecode.AsSeconds())
+			if (bFirstElement)
 			{
-				ResultSampleTime.Timecode = OldestSampleTime.Timecode;
+				ResultSampleTime = OldestSampleTime;
+				bFirstElement = false;
+			}
+			else
+			{
+				ResultSampleTime.PlatformSecond = FMath::Min(OldestSampleTime.PlatformSecond, ResultSampleTime.PlatformSecond);
+				if (OldestSampleTime.Timecode.AsSeconds() < ResultSampleTime.Timecode.AsSeconds())
+				{
+					ResultSampleTime.Timecode = OldestSampleTime.Timecode;
+				}
 			}
 		}
 	}
@@ -777,7 +448,7 @@ FTimedDataChannelSampleTime UTimedDataMonitorSubsystem::GetInputNewestDataTime(c
 	BuildSourcesListIfNeeded();
 
 	FTimedDataChannelSampleTime ResultSampleTime(0.0, FQualifiedFrameTime());
-	if (FTimeDataInputItem* SourceItem = InputMap.Find(Identifier))
+	if (const FTimeDataInputItem* SourceItem = InputMap.Find(Identifier))
 	{
 		for (FTimedDataMonitorChannelIdentifier ChannelIdentifier : SourceItem->ChannelIdentifiers)
 		{
@@ -798,7 +469,7 @@ bool UTimedDataMonitorSubsystem::IsDataBufferSizeControlledByInput(const FTimedD
 {
 	BuildSourcesListIfNeeded();
 
-	if (FTimeDataInputItem* SourceItem = InputMap.Find(Identifier))
+	if (const FTimeDataInputItem* SourceItem = InputMap.Find(Identifier))
 	{
 		return SourceItem->Input->IsDataBufferSizeControlledByInput();
 	}
@@ -811,7 +482,7 @@ int32 UTimedDataMonitorSubsystem::GetInputDataBufferSize(const FTimedDataMonitor
 {
 	BuildSourcesListIfNeeded();
 
-	if (FTimeDataInputItem* SourceItem = InputMap.Find(Identifier))
+	if (const FTimeDataInputItem* SourceItem = InputMap.Find(Identifier))
 	{
 		return SourceItem->Input->GetDataBufferSize();
 	}
@@ -826,17 +497,20 @@ void UTimedDataMonitorSubsystem::SetInputDataBufferSize(const FTimedDataMonitorI
 
 	if (FTimeDataInputItem* SourceItem = InputMap.Find(Identifier))
 	{
-		return SourceItem->Input->SetDataBufferSize(BufferSize);
+		if (SourceItem->Input->IsDataBufferSizeControlledByInput())
+		{
+			SourceItem->Input->SetDataBufferSize(BufferSize);
+		}
 	}
 }
 
-ETimedDataInputState UTimedDataMonitorSubsystem::GetInputState(const FTimedDataMonitorInputIdentifier& Identifier)
+ETimedDataInputState UTimedDataMonitorSubsystem::GetInputConnectionState(const FTimedDataMonitorInputIdentifier& Identifier)
 {
 	BuildSourcesListIfNeeded();
 
 	ETimedDataInputState WorstState = ETimedDataInputState::Connected;
 	bool bHasAtLeastOneItem = false;
-	if (FTimeDataInputItem* InputItem = InputMap.Find(Identifier))
+	if (const FTimeDataInputItem* InputItem = InputMap.Find(Identifier))
 	{
 		for (const FTimedDataMonitorChannelIdentifier& ChannelIdentifier : InputItem->ChannelIdentifiers)
 		{
@@ -862,12 +536,72 @@ ETimedDataInputState UTimedDataMonitorSubsystem::GetInputState(const FTimedDataM
 }
 
 
+ETimedDataMonitorEvaluationState UTimedDataMonitorSubsystem::GetInputEvaluationState(const FTimedDataMonitorInputIdentifier& Identifier)
+{
+	BuildSourcesListIfNeeded();
+
+	ETimedDataMonitorEvaluationState WorstState = ETimedDataMonitorEvaluationState::NoSample;
+	if (const FTimeDataInputItem* InputItem = InputMap.Find(Identifier))
+	{
+		if (InputItem->ChannelIdentifiers.Num() > 0)
+		{
+			WorstState = ETimedDataMonitorEvaluationState::Disabled;
+			for (const FTimedDataMonitorChannelIdentifier& ChannelIdentifier : InputItem->ChannelIdentifiers)
+			{
+				const ETimedDataMonitorEvaluationState ChannelState = GetChannelEvaluationState(ChannelIdentifier);
+				uint8 ChannelValue = FMath::Min((uint8)ChannelState, (uint8)WorstState);
+				WorstState = (ETimedDataMonitorEvaluationState)ChannelValue;
+			}
+		}
+
+	}
+
+	return WorstState;
+}
+
+
+float UTimedDataMonitorSubsystem::GetInputEvaluationDistanceToNewestSampleMean(const FTimedDataMonitorInputIdentifier& Identifier)
+{
+	BuildSourcesListIfNeeded();
+	
+	float WorstNewestMean = 0.f;
+	if (const FTimeDataInputItem* InputItem = InputMap.Find(Identifier))
+	{
+		WorstNewestMean = TNumericLimits<float>::Lowest();
+		for (const FTimedDataMonitorChannelIdentifier& ChannelIdentifier : InputItem->ChannelIdentifiers)
+		{
+			WorstNewestMean = FMath::Max(ChannelMap[ChannelIdentifier].Statistics.IncrementalAverageNewestDistance, WorstNewestMean);
+		}
+	}
+
+	return WorstNewestMean;
+}
+
+
+float UTimedDataMonitorSubsystem::GetInputEvaluationDistanceToOldestSampleMean(const FTimedDataMonitorInputIdentifier& Identifier)
+{
+	BuildSourcesListIfNeeded();
+
+	float WorstOldesttMean = 0.f;
+	if (const FTimeDataInputItem* InputItem = InputMap.Find(Identifier))
+	{
+		WorstOldesttMean = TNumericLimits<float>::Max();
+		for (const FTimedDataMonitorChannelIdentifier& ChannelIdentifier : InputItem->ChannelIdentifiers)
+		{
+			WorstOldesttMean = FMath::Min(ChannelMap[ChannelIdentifier].Statistics.IncrementalAverageOldestDistance, WorstOldesttMean);
+		}
+	}
+
+	return WorstOldesttMean;
+}
+
+
 float UTimedDataMonitorSubsystem::GetInputEvaluationDistanceToNewestSampleStandardDeviation(const FTimedDataMonitorInputIdentifier& Identifier)
 {
 	BuildSourcesListIfNeeded();
 
 	float WorstNewestSSD = 0.f;
-	if (FTimeDataInputItem* InputItem = InputMap.Find(Identifier))
+	if (const FTimeDataInputItem* InputItem = InputMap.Find(Identifier))
 	{
 		for (const FTimedDataMonitorChannelIdentifier& ChannelIdentifier : InputItem->ChannelIdentifiers)
 		{
@@ -884,7 +618,7 @@ float UTimedDataMonitorSubsystem::GetInputEvaluationDistanceToOldestSampleStanda
 	BuildSourcesListIfNeeded();
 
 	float WorstOldestSSD = 0.f;
-	if (FTimeDataInputItem* InputItem = InputMap.Find(Identifier))
+	if (const FTimeDataInputItem* InputItem = InputMap.Find(Identifier))
 	{
 		for (const FTimedDataMonitorChannelIdentifier& ChannelIdentifier : InputItem->ChannelIdentifiers)
 		{
@@ -908,7 +642,7 @@ bool UTimedDataMonitorSubsystem::IsChannelEnabled(const FTimedDataMonitorChannel
 {
 	BuildSourcesListIfNeeded();
 
-	if (FTimeDataChannelItem* SourceItem = ChannelMap.Find(Identifier))
+	if (const FTimeDataChannelItem* SourceItem = ChannelMap.Find(Identifier))
 	{
 		return SourceItem->bEnabled;
 	}
@@ -932,7 +666,7 @@ FTimedDataMonitorInputIdentifier UTimedDataMonitorSubsystem::GetChannelInput(con
 {
 	BuildSourcesListIfNeeded();
 
-	if (FTimeDataChannelItem* SourceItem = ChannelMap.Find(Identifier))
+	if (const FTimeDataChannelItem* SourceItem = ChannelMap.Find(Identifier))
 	{
 		return SourceItem->InputIdentifier;
 	}
@@ -945,7 +679,7 @@ FText UTimedDataMonitorSubsystem::GetChannelDisplayName(const FTimedDataMonitorC
 {
 	BuildSourcesListIfNeeded();
 
-	if (FTimeDataChannelItem* SourceItem = ChannelMap.Find(Identifier))
+	if (const FTimeDataChannelItem* SourceItem = ChannelMap.Find(Identifier))
 	{
 		return SourceItem->Channel->GetDisplayName();
 	}
@@ -954,11 +688,11 @@ FText UTimedDataMonitorSubsystem::GetChannelDisplayName(const FTimedDataMonitorC
 }
 
 
-ETimedDataInputState UTimedDataMonitorSubsystem::GetChannelState(const FTimedDataMonitorChannelIdentifier& Identifier)
+ETimedDataInputState UTimedDataMonitorSubsystem::GetChannelConnectionState(const FTimedDataMonitorChannelIdentifier& Identifier)
 {
 	BuildSourcesListIfNeeded();
 
-	if (FTimeDataChannelItem* SourceItem = ChannelMap.Find(Identifier))
+	if (const FTimeDataChannelItem* SourceItem = ChannelMap.Find(Identifier))
 	{
 		return SourceItem->Channel->GetState();
 	}
@@ -967,11 +701,42 @@ ETimedDataInputState UTimedDataMonitorSubsystem::GetChannelState(const FTimedDat
 }
 
 
+ETimedDataMonitorEvaluationState UTimedDataMonitorSubsystem::GetChannelEvaluationState(const FTimedDataMonitorChannelIdentifier& Identifier)
+{
+	BuildSourcesListIfNeeded();
+
+	if (const FTimeDataChannelItem* SourceItem = ChannelMap.Find(Identifier))
+	{
+		if (SourceItem->Channel->GetState() != ETimedDataInputState::Connected || !SourceItem->bEnabled)
+		{
+			return ETimedDataMonitorEvaluationState::Disabled;
+		}
+		if (SourceItem->Channel->GetNumberOfSamples() <= 0)
+		{
+			return ETimedDataMonitorEvaluationState::NoSample;
+		}
+
+		const ITimedDataInput* Input = InputMap[SourceItem->InputIdentifier].Input;
+		check(Input);
+		const ETimedDataInputEvaluationType EvaluationType = Input->GetEvaluationType();
+		const double EvaluationOffset = Input->GetEvaluationOffsetInSeconds();
+		const double OldestSampleTime = SourceItem->Channel->GetOldestDataTime().AsSeconds(EvaluationType);
+		const double NewstedSampleTime = SourceItem->Channel->GetNewestDataTime().AsSeconds(EvaluationType);
+		const double EvaluationTime = GetEvaluationTime(EvaluationType);
+		const double OffsettedEvaluationTime = EvaluationTime - EvaluationOffset;
+		bool bIsInRange = (OffsettedEvaluationTime >= OldestSampleTime) && (OffsettedEvaluationTime <= NewstedSampleTime);
+		return bIsInRange ? ETimedDataMonitorEvaluationState::InsideRange : ETimedDataMonitorEvaluationState::OutsideRange;
+	}
+
+	return ETimedDataMonitorEvaluationState::Disabled;
+}
+
+
 FTimedDataChannelSampleTime UTimedDataMonitorSubsystem::GetChannelOldestDataTime(const FTimedDataMonitorChannelIdentifier& Identifier)
 {
 	BuildSourcesListIfNeeded();
 
-	if (FTimeDataChannelItem* SourceItem = ChannelMap.Find(Identifier))
+	if (const FTimeDataChannelItem* SourceItem = ChannelMap.Find(Identifier))
 	{
 		return SourceItem->Channel->GetOldestDataTime();
 	}
@@ -984,7 +749,7 @@ FTimedDataChannelSampleTime UTimedDataMonitorSubsystem::GetChannelNewestDataTime
 {
 	BuildSourcesListIfNeeded();
 
-	if (FTimeDataChannelItem* SourceItem = ChannelMap.Find(Identifier))
+	if (const FTimeDataChannelItem* SourceItem = ChannelMap.Find(Identifier))
 	{
 		return SourceItem->Channel->GetNewestDataTime();
 	}
@@ -997,7 +762,7 @@ int32 UTimedDataMonitorSubsystem::GetChannelNumberOfSamples(const FTimedDataMoni
 {
 	BuildSourcesListIfNeeded();
 
-	if (FTimeDataChannelItem* SourceItem = ChannelMap.Find(Identifier))
+	if (const FTimeDataChannelItem* SourceItem = ChannelMap.Find(Identifier))
 	{
 		return SourceItem->Channel->GetNumberOfSamples();
 	}
@@ -1010,9 +775,13 @@ int32 UTimedDataMonitorSubsystem::GetChannelDataBufferSize(const FTimedDataMonit
 {
 	BuildSourcesListIfNeeded();
 
-	if (FTimeDataChannelItem* SourceItem = ChannelMap.Find(Identifier))
+	if (const FTimeDataChannelItem* SourceItem = ChannelMap.Find(Identifier))
 	{
 		if (InputMap[SourceItem->InputIdentifier].Input->IsDataBufferSizeControlledByInput())
+		{
+			return InputMap[SourceItem->InputIdentifier].Input->GetDataBufferSize();
+		}
+		else
 		{
 			return SourceItem->Channel->GetDataBufferSize();
 		}
@@ -1028,7 +797,7 @@ void UTimedDataMonitorSubsystem::SetChannelDataBufferSize(const FTimedDataMonito
 
 	if (FTimeDataChannelItem* SourceItem = ChannelMap.Find(Identifier))
 	{
-		if (InputMap[SourceItem->InputIdentifier].Input->IsDataBufferSizeControlledByInput())
+		if (!InputMap[SourceItem->InputIdentifier].Input->IsDataBufferSizeControlledByInput())
 		{
 			SourceItem->Channel->SetDataBufferSize(BufferSize);
 		}
@@ -1040,7 +809,7 @@ int32 UTimedDataMonitorSubsystem::GetChannelBufferUnderflowStat(const FTimedData
 {
 	BuildSourcesListIfNeeded();
 
-	if (FTimeDataChannelItem* SourceItem = ChannelMap.Find(Identifier))
+	if (const FTimeDataChannelItem* SourceItem = ChannelMap.Find(Identifier))
 	{
 		return SourceItem->Channel->GetBufferUnderflowStat();
 	}
@@ -1052,7 +821,7 @@ int32 UTimedDataMonitorSubsystem::GetChannelBufferOverflowStat(const FTimedDataM
 {
 	BuildSourcesListIfNeeded();
 
-	if (FTimeDataChannelItem* SourceItem = ChannelMap.Find(Identifier))
+	if (const FTimeDataChannelItem* SourceItem = ChannelMap.Find(Identifier))
 	{
 		return SourceItem->Channel->GetBufferOverflowStat();
 	}
@@ -1064,7 +833,7 @@ int32 UTimedDataMonitorSubsystem::GetChannelFrameDroppedStat(const FTimedDataMon
 {
 	BuildSourcesListIfNeeded();
 
-	if (FTimeDataChannelItem* SourceItem = ChannelMap.Find(Identifier))
+	if (const FTimeDataChannelItem* SourceItem = ChannelMap.Find(Identifier))
 	{
 		return SourceItem->Channel->GetFrameDroppedStat();
 	}
@@ -1076,7 +845,7 @@ float UTimedDataMonitorSubsystem::GetChannelEvaluationDistanceToNewestSampleMean
 {
 	BuildSourcesListIfNeeded();
 
-	if (FTimeDataChannelItem* SourceItem = ChannelMap.Find(Identifier))
+	if (const FTimeDataChannelItem* SourceItem = ChannelMap.Find(Identifier))
 	{
 		return SourceItem->Statistics.IncrementalAverageNewestDistance;
 	}
@@ -1088,7 +857,7 @@ float UTimedDataMonitorSubsystem::GetChannelEvaluationDistanceToOldestSampleMean
 {
 	BuildSourcesListIfNeeded();
 
-	if (FTimeDataChannelItem* SourceItem = ChannelMap.Find(Identifier))
+	if (const FTimeDataChannelItem* SourceItem = ChannelMap.Find(Identifier))
 	{
 		return SourceItem->Statistics.IncrementalAverageOldestDistance;
 	}
@@ -1100,7 +869,7 @@ float UTimedDataMonitorSubsystem::GetChannelEvaluationDistanceToNewestSampleStan
 {
 	BuildSourcesListIfNeeded();
 
-	if (FTimeDataChannelItem* SourceItem = ChannelMap.Find(Identifier))
+	if (const FTimeDataChannelItem* SourceItem = ChannelMap.Find(Identifier))
 	{
 		return SourceItem->Statistics.DistanceToNewestSTD;
 	}
@@ -1112,7 +881,7 @@ float UTimedDataMonitorSubsystem::GetChannelEvaluationDistanceToOldestSampleStan
 {
 	BuildSourcesListIfNeeded();
 
-	if (FTimeDataChannelItem* SourceItem = ChannelMap.Find(Identifier))
+	if (const FTimeDataChannelItem* SourceItem = ChannelMap.Find(Identifier))
 	{
 		return SourceItem->Statistics.DistanceToOldestSTD;
 	}
@@ -1124,7 +893,7 @@ void UTimedDataMonitorSubsystem::GetChannelLastEvaluationDataStat(const FTimedDa
 {
 	BuildSourcesListIfNeeded();
 
-	if (FTimeDataChannelItem* SourceItem = ChannelMap.Find(Identifier))
+	if (const FTimeDataChannelItem* SourceItem = ChannelMap.Find(Identifier))
 	{
 		SourceItem->Channel->GetLastEvaluationData(Result);
 	}
@@ -1294,36 +1063,6 @@ void UTimedDataMonitorSubsystem::OnTimedDataSourceCollectionChanged()
 }
 
 
-double UTimedDataMonitorSubsystem::GetSeconds(ETimedDataInputEvaluationType Evaluation, const FTimedDataChannelSampleTime& SampleTime) const
-{
-	return Evaluation == ETimedDataInputEvaluationType::Timecode ? SampleTime.Timecode.AsSeconds() : SampleTime.PlatformSecond;
-}
-
-
-double UTimedDataMonitorSubsystem::CalculateAverageInDeltaTimeBetweenSample(ETimedDataInputEvaluationType Evaluation, const TArray<FTimedDataChannelSampleTime>& SampleTimes) const
-{
-	double Average = 0.0;
-	if (SampleTimes.Num() >= 2)
-	{
-		// Get the average of the last 10 samples in seconds
-		const int32 AvgCounter = FMath::Min(SampleTimes.Num()-1, 10-1);
-		
-		const int32 SampleTimeNum = SampleTimes.Num();
-
-		for (int32 Index = 1; Index <= AvgCounter; ++Index)
-		{
-			double Delta = GetSeconds(Evaluation, SampleTimes[SampleTimeNum - Index]) - GetSeconds(Evaluation, SampleTimes[SampleTimeNum - Index - 1]);
-			Average += (Delta - Average) / (double)Index;
-		}
-	}
-	else
-	{
-		Average = FApp::GetDeltaTime(); // was not able to find a correct delta time. guess one.
-	}
-	return Average;
-}
-
-
 void UTimedDataMonitorSubsystem::Tick(float DeltaTime)
 {
 	const bool bUpdateStats = CVarEnableTimedDataMonitorSubsystemStats.GetValueOnGameThread();
@@ -1385,6 +1124,9 @@ void FTimedDataChannelEvaluationStatistics::Update(float DistanceToOldest, float
 	//Square root of that average gives us sigma (standard deviation)
 	DistanceToNewestSTD = FMath::Sqrt(IncrementalVarianceDistanceNewest);
 	DistanceToOldestSTD = FMath::Sqrt(IncrementalVarianceDistanceOldest);
+
+	LastDistanceToOldest = DistanceToOldest;
+	LastDistanceToNewest = DistanceToNewest;
 }
 
 void FTimedDataChannelEvaluationStatistics::Reset()
diff --git a/Engine/Plugins/VirtualProduction/TimedDataMonitor/Source/TimedDataMonitor/Public/TimedDataMonitorSubsystem.h b/Engine/Plugins/VirtualProduction/TimedDataMonitor/Source/TimedDataMonitor/Public/TimedDataMonitorSubsystem.h
--- a/Engine/Plugins/VirtualProduction/TimedDataMonitor/Source/TimedDataMonitor/Public/TimedDataMonitorSubsystem.h
+++ b/Engine/Plugins/VirtualProduction/TimedDataMonitor/Source/TimedDataMonitor/Public/TimedDataMonitorSubsystem.h
@@ -4,55 +4,16 @@
 
 #include "CoreMinimal.h"
 #include "Engine/EngineBaseTypes.h"
+#include "Engine/LatentActionManager.h"
 #include "ITimedDataInput.h"
 #include "Subsystems/EngineSubsystem.h"
 #include "Tickable.h"
+#include "TimedDataMonitorCalibration.h"
+#include "TimedDataMonitorTypes.h"
 
 #include "TimedDataMonitorSubsystem.generated.h"
 
 
-USTRUCT(BlueprintType)
-struct TIMEDDATAMONITOR_API FTimedDataMonitorInputIdentifier
-{
-	GENERATED_BODY()
-
-private:
-	UPROPERTY()
-	FGuid Identifier;
-
-public:
-	static FTimedDataMonitorInputIdentifier NewIdentifier();
-
-	bool IsValid() const { return Identifier.IsValid(); }
-
-	bool operator== (const FTimedDataMonitorInputIdentifier& Other) const { return Other.Identifier == Identifier; }
-	bool operator!= (const FTimedDataMonitorInputIdentifier& Other) const { return Other.Identifier != Identifier; }
-
-	friend uint32 GetTypeHash(const FTimedDataMonitorInputIdentifier& InIdentifier) { return GetTypeHash(InIdentifier.Identifier); }
-};
-
-
-USTRUCT(BlueprintType)
-struct TIMEDDATAMONITOR_API FTimedDataMonitorChannelIdentifier
-{
-	GENERATED_BODY()
-
-private:
-	UPROPERTY()
-	FGuid Identifier;
-
-public:
-	static FTimedDataMonitorChannelIdentifier NewIdentifier();
-
-	bool IsValid() const { return Identifier.IsValid(); }
-
-	bool operator== (const FTimedDataMonitorChannelIdentifier& Other) const { return Other.Identifier == Identifier; }
-	bool operator!= (const FTimedDataMonitorChannelIdentifier& Other) const { return Other.Identifier != Identifier; }
-
-	friend uint32 GetTypeHash(const FTimedDataMonitorChannelIdentifier& InIdentifier) { return GetTypeHash(InIdentifier.Identifier); }
-};
-
-
 UENUM()
 enum class ETimedDataMonitorInputEnabled : uint8
 {
@@ -63,70 +24,12 @@ enum class ETimedDataMonitorInputEnabled : uint8
 
 
 UENUM()
-enum class ETimedDataMonitorCalibrationReturnCode : uint8
-{
-	/** Success. The values were synchronized. */
-	Succeeded,
-	/** Failed. The timecode provider doesn't have a proper timecode value. */
-	Failed_NoTimecode,
-	/** Failed. At least one input is unresponsive. */
-	Failed_UnresponsiveInput,
-	/** Failed. At least one input doesn't have a defined frame rate. */
-	Failed_InvalidFrameRate,
-	/** Failed. At least one input doesn't have data buffered. */
-	Failed_NoDataBuffered,
-	/** Failed. We tried to find a valid offset for the timecode provider and failed. */
-	Failed_CanNotCallibrateWithoutJam,
-	/** It failed but, we increased the number of buffer of at least one channel. You may retry to see if it works now. */
-	Retry_BufferSizeHasBeenIncreased,
-};
-
-
-USTRUCT(BlueprintType)
-struct FTimedDataMonitorCalibrationResult
-{
-	GENERATED_BODY();
-
-	UPROPERTY(BlueprintReadOnly, VisibleAnywhere, Category = "Result")
-	ETimedDataMonitorCalibrationReturnCode ReturnCode;
-
-	UPROPERTY(BlueprintReadOnly, VisibleAnywhere, Category = "Result")
-	TArray<FTimedDataMonitorChannelIdentifier> FailureChannelIdentifiers;
-};
-
-
-UENUM()
-enum class ETimedDataMonitorJamReturnCode : uint8
+enum class ETimedDataMonitorEvaluationState : uint8
 {
-	/** Success. The values were synchronized. */
-	Succeeded,
-	/** Failed. The timecode provider was not existing or not synchronized. */
-	Failed_NoTimecode,
-	/** Failed. At least one channel is unresponsive. */
-	Failed_UnresponsiveInput,
-	/** Failed. The evaluation type of at least of input doesn't match with what was requested. */
-	Failed_EvaluationTypeDoNotMatch,
-	/** Failed. The channel doesn't have any data in it's buffer to synchronized with. */
-	Failed_NoDataBuffered,
-	/** Failed. The channel or buffer size has been increase but it's still not enough. */
-	Failed_BufferSizeHaveBeenMaxed,
-	/** Failed but we increased the number of buffer of at least one channel. You may retry to see if it works now. */
-	Retry_BufferSizeHasBeenIncreased,
-};
-
-USTRUCT(BlueprintType)
-struct FTimedDataMonitorJamResult
-{
-	GENERATED_BODY();
-
-	UPROPERTY(BlueprintReadOnly, VisibleAnywhere, Category="Result")
-	ETimedDataMonitorJamReturnCode ReturnCode;
-
-	UPROPERTY(BlueprintReadOnly, VisibleAnywhere, Category = "Result")
-	TArray<FTimedDataMonitorInputIdentifier> FailureInputIdentifiers;
-
-	UPROPERTY(BlueprintReadOnly, VisibleAnywhere, Category = "Result")
-	TArray<FTimedDataMonitorChannelIdentifier> FailureChannelIdentifiers;
+	NoSample = 0,
+	OutsideRange = 1,
+	InsideRange = 2,
+	Disabled = 3,
 };
 
 
@@ -141,6 +44,12 @@ struct FTimedDataChannelEvaluationStatistics
 	/** Current number of samples used for the running average and variance */
 	int32 SampleCount = 0;
 
+	/** Last distance from evaluation time to oldest sample used for statistics */
+	float LastDistanceToOldest = 0.0f;
+
+	/** Last distance from evaluation time to newest sample used for statistics */
+	float LastDistanceToNewest = 0.0f;
+
 	/** Running average of the distance between last evaluation time and oldest sample in the buffer */
 	float IncrementalAverageOldestDistance = 0.0f;
 
@@ -229,6 +138,9 @@ public:
 	/** Get the interface for a specific channel identifier. */
 	ITimedDataInputChannel* GetTimedDataChannel(const FTimedDataMonitorChannelIdentifier& Identifier);
 
+	/** Get the current evaluation time. */
+	static double GetEvaluationTime(ETimedDataInputEvaluationType EvaluationType);
+
 public:
 	/** Get the list of all the inputs. */
 	UFUNCTION(BlueprintCallable, Category = "Timed Data Monitor")
@@ -238,21 +150,26 @@ public:
 	UFUNCTION(BlueprintCallable, Category = "Timed Data Monitor")
 	TArray<FTimedDataMonitorChannelIdentifier> GetAllChannels();
 
-	/**
-	 * Set evaluation offset for all inputs and the engine's timecode provider to align all the buffers.
-	 * If there is no data available, it may increase the buffer size of an input.
-	 */
+	/** Get the list of all the channels that are enabled. */
 	UFUNCTION(BlueprintCallable, Category = "Timed Data Monitor")
-	FTimedDataMonitorCalibrationResult CalibrateWithTimecodeProvider();
+	TArray<FTimedDataMonitorChannelIdentifier> GetAllEnabledChannels();
+
+	/** Change the Timecode Provider offset to align all inputs and channels. */
+	UFUNCTION(BlueprintCallable, Category = "Timed Data Monitor", meta = (DisplayName = "Calibrate", Latent, LatentInfo = "LatentInfo", WorldContext = "WorldContextObject"))
+	void CalibrateLatent(UObject* WorldContextObject, FLatentActionInfo LatentInfo, const FTimedDataMonitorCalibrationParameters& CalibrationParameters, FTimedDataMonitorCalibrationResult& Result);
 
 	/** Assume all data samples were produce at the same time and align them with the current platform's time */
 	UFUNCTION(BlueprintCallable, Category = "Timed Data Monitor")
-	FTimedDataMonitorJamResult JamInputs(ETimedDataInputEvaluationType EvaluationType);
+	FTimedDataMonitorTimeCorrectionResult ApplyTimeCorrection(const FTimedDataMonitorInputIdentifier& Identifier, const FTimedDataMonitorTimeCorrectionParameters& TimeCorrectionParameters);
 
 	/** Reset the stat of all the inputs. */
 	UFUNCTION(BlueprintCallable, Category = "Timed Data Monitor")
 	void ResetAllBufferStats();
 
+	/** Get the worst evaluation state of all the inputs. */
+	UFUNCTION(BlueprintCallable, Category = "Timed Data Monitor")
+	ETimedDataMonitorEvaluationState GetEvaluationState();
+
 	/** Return true if the identifier is a valid input. */
 	UFUNCTION(BlueprintCallable, Category = "Timed Data Monitor|Input")
 	bool DoesInputExist(const FTimedDataMonitorInputIdentifier& Identifier);
@@ -313,9 +230,21 @@ public:
 	UFUNCTION(BlueprintCallable, Category = "Timed Data Monitor|Input")
 	void SetInputDataBufferSize(const FTimedDataMonitorInputIdentifier& Identifier, int32 BufferSize);
 
-	/** Get the worst state of all the channel state of that input. */
+	/** Get the worst state of all the channels of that input. */
+	UFUNCTION(BlueprintCallable, Category = "Timed Data Monitor|Input")
+	ETimedDataInputState GetInputConnectionState(const FTimedDataMonitorInputIdentifier& Identifier);
+
+	/** Get the worst evaluation state of all the channels of that input. */
 	UFUNCTION(BlueprintCallable, Category = "Timed Data Monitor|Input")
-	ETimedDataInputState GetInputState(const FTimedDataMonitorInputIdentifier& Identifier);
+	ETimedDataMonitorEvaluationState GetInputEvaluationState(const FTimedDataMonitorInputIdentifier& Identifier);
+
+	/** Returns the max average distance, in seconds, between evaluation time and newest sample */
+	UFUNCTION(BlueprintCallable, Category = "Timed Data Monitor|Input")
+	float GetInputEvaluationDistanceToNewestSampleMean(const FTimedDataMonitorInputIdentifier& Identifier);
+
+	/** Returns the min average distance, in seconds, between evaluation time and oldest sample */
+	UFUNCTION(BlueprintCallable, Category = "Timed Data Monitor|Input")
+	float GetInputEvaluationDistanceToOldestSampleMean(const FTimedDataMonitorInputIdentifier& Identifier);
 
 	/** Returns the standard deviation of the distance, in seconds, between evaluation time and newest sample */
 	UFUNCTION(BlueprintCallable, Category = "Timed Data Monitor|Input")
@@ -348,9 +277,13 @@ public:
 	UFUNCTION(BlueprintCallable, Category = "Timed Data Monitor|Channel")
 	FText GetChannelDisplayName(const FTimedDataMonitorChannelIdentifier& Identifier);
 
-	/** Get the worst state of all the input state of that channel. */
+	/** Get the state the channel. */
 	UFUNCTION(BlueprintCallable, Category = "Timed Data Monitor|Channel")
-	ETimedDataInputState GetChannelState (const FTimedDataMonitorChannelIdentifier& Identifier);
+	ETimedDataInputState GetChannelConnectionState (const FTimedDataMonitorChannelIdentifier& Identifier);
+
+	/** Get the evaluation state of the channel. */
+	UFUNCTION(BlueprintCallable, Category = "Timed Data Monitor|Channel")
+	ETimedDataMonitorEvaluationState GetChannelEvaluationState(const FTimedDataMonitorChannelIdentifier& Identifier);
 
 	/** Get the channel oldest sample time. */
 	UFUNCTION(BlueprintCallable, Category = "Timed Data Monitor|Channel")
@@ -412,9 +345,6 @@ private:
 
 	void OnTimedDataSourceCollectionChanged();
 
-	double GetSeconds(ETimedDataInputEvaluationType Evaluation, const FTimedDataChannelSampleTime& SampleTime) const;
-	double CalculateAverageInDeltaTimeBetweenSample(ETimedDataInputEvaluationType Evaluation, const TArray<FTimedDataChannelSampleTime>& SampleTimes) const;
-
 	/** Update internal statistics for each enabled channel */
 	void UpdateEvaluationStatistics();
 
diff --git a/Engine/Plugins/VirtualProduction/TimedDataMonitor/Source/TimedDataMonitor/TimedDataMonitor.Build.cs b/Engine/Plugins/VirtualProduction/TimedDataMonitor/Source/TimedDataMonitor/TimedDataMonitor.Build.cs
--- a/Engine/Plugins/VirtualProduction/TimedDataMonitor/Source/TimedDataMonitor/TimedDataMonitor.Build.cs
+++ b/Engine/Plugins/VirtualProduction/TimedDataMonitor/Source/TimedDataMonitor/TimedDataMonitor.Build.cs
@@ -10,15 +10,12 @@ public class TimedDataMonitor : ModuleRules
 			new string[] {
 				"Core",
 				"CoreUObject",
-				"SlateCore",
 				"TimeManagement",
 			});
 
 		PrivateDependencyModuleNames.AddRange(
 			new string[] {
 				"Engine",
-				"Projects",
-				"Slate",
 			});
 	}
 }
diff --git a/Engine/Plugins/VirtualProduction/TimedDataMonitor/Source/TimedDataMonitorEditor/Private/STimedDataListView.cpp b/Engine/Plugins/VirtualProduction/TimedDataMonitor/Source/TimedDataMonitorEditor/Private/STimedDataListView.cpp
--- a/Engine/Plugins/VirtualProduction/TimedDataMonitor/Source/TimedDataMonitorEditor/Private/STimedDataListView.cpp
+++ b/Engine/Plugins/VirtualProduction/TimedDataMonitor/Source/TimedDataMonitorEditor/Private/STimedDataListView.cpp
@@ -13,14 +13,18 @@
 
 #include "EditorFontGlyphs.h"
 #include "EditorStyleSet.h"
+#include "TimedDataMonitorEditorStyle.h"
 
+#include "Framework/MultiBox/MultiBoxBuilder.h"
+#include "STimedDataMonitorPanel.h"
+#include "STimedDataNumericEntryBox.h"
+#include "STimingDiagramWidget.h"
 #include "Widgets/Images/SImage.h"
 #include "Widgets/Input/SCheckBox.h"
 #include "Widgets/Input/SNumericEntryBox.h"
 #include "Widgets/Text/STextBlock.h"
 #include "Widgets/Views/SExpanderArrow.h"
 #include "Widgets/SNullWidget.h"
-#include "Widgets/TimedDataMonitorStyle.h"
 
 
 #define LOCTEXT_NAMESPACE "STimedDataListView"
@@ -30,6 +34,7 @@ namespace TimedDataListView
 {
 	const FName HeaderIdName_Enable			= "Enable";
 	const FName HeaderIdName_Icon			= "Edit";
+	const FName HeaderIdName_EvaluationMode = "EvaluationMode";
 	const FName HeaderIdName_Name			= "Name";
 	const FName HeaderIdName_Description	= "Description";
 	const FName HeaderIdName_TimeCorrection	= "TimeCorrection";
@@ -54,14 +59,14 @@ namespace TimedDataListView
 struct FTimedDataInputTableRowData : TSharedFromThis<FTimedDataInputTableRowData>
 {
 	FTimedDataInputTableRowData(const FTimedDataMonitorInputIdentifier& InInputId)
-		: InputIdentifer(InInputId), bIsInput(true)
+		: InputIdentifier(InInputId), bIsInput(true)
 	{
 		UTimedDataMonitorSubsystem* TimedDataMonitorSubsystem = GEngine->GetEngineSubsystem<UTimedDataMonitorSubsystem>();
 		check(TimedDataMonitorSubsystem);
 
-		DisplayName = TimedDataMonitorSubsystem->GetInputDisplayName(InputIdentifer);
+		DisplayName = TimedDataMonitorSubsystem->GetInputDisplayName(InputIdentifier);
 
-		if (ITimedDataInput* InpuTimedData = TimedDataMonitorSubsystem->GetTimedDataInput(InputIdentifer))
+		if (ITimedDataInput* InpuTimedData = TimedDataMonitorSubsystem->GetTimedDataInput(InputIdentifier))
 		{
 			InputIcon = InpuTimedData->GetDisplayIcon();
 		}
@@ -73,7 +78,7 @@ struct FTimedDataInputTableRowData : TSharedFromThis<FTimedDataInputTableRowData
 		UTimedDataMonitorSubsystem* TimedDataMonitorSubsystem = GEngine->GetEngineSubsystem<UTimedDataMonitorSubsystem>();
 		check(TimedDataMonitorSubsystem);
 
-		InputIdentifer = TimedDataMonitorSubsystem->GetChannelInput(ChannelIdentifier);
+		InputIdentifier = TimedDataMonitorSubsystem->GetChannelInput(ChannelIdentifier);
 		DisplayName = TimedDataMonitorSubsystem->GetChannelDisplayName(ChannelIdentifier);
 	}
 
@@ -85,7 +90,7 @@ struct FTimedDataInputTableRowData : TSharedFromThis<FTimedDataInputTableRowData
 		FTimedDataChannelSampleTime NewestDataTime;
 		if (bIsInput)
 		{
-			switch (TimedDataMonitorSubsystem->GetInputEnabled(InputIdentifer))
+			switch (TimedDataMonitorSubsystem->GetInputEnabled(InputIdentifier))
 			{
 			case ETimedDataMonitorInputEnabled::Enabled:
 				CachedEnabled = ECheckBoxState::Checked;
@@ -99,13 +104,15 @@ struct FTimedDataInputTableRowData : TSharedFromThis<FTimedDataInputTableRowData
 				break;
 			};
 
-			CachedInputEvaluationType = TimedDataMonitorSubsystem->GetInputEvaluationType(InputIdentifer);
-			CachedInputEvaluationOffset = TimedDataMonitorSubsystem->GetInputEvaluationOffsetInSeconds(InputIdentifer);
-			CachedState = TimedDataMonitorSubsystem->GetInputState(InputIdentifer);
-			CachedBufferSize = TimedDataMonitorSubsystem->GetInputDataBufferSize(InputIdentifer);
-			bCachedCanEditBufferSize = (CachedEnabled == ECheckBoxState::Checked || CachedEnabled == ECheckBoxState::Undetermined) && TimedDataMonitorSubsystem->IsDataBufferSizeControlledByInput(InputIdentifer);
+			CachedInputEvaluationType = TimedDataMonitorSubsystem->GetInputEvaluationType(InputIdentifier);
+			CachedInputEvaluationOffset = TimedDataMonitorSubsystem->GetInputEvaluationOffsetInSeconds(InputIdentifier);
+			CachedState = TimedDataMonitorSubsystem->GetInputConnectionState(InputIdentifier);
+			CachedBufferSize = TimedDataMonitorSubsystem->GetInputDataBufferSize(InputIdentifier);
+			CachedCurrentAmountOfBuffer = 0;
+			bControlBufferSize = TimedDataMonitorSubsystem->IsDataBufferSizeControlledByInput(InputIdentifier);
+			bCachedCanEditBufferSize = (CachedEnabled == ECheckBoxState::Checked || CachedEnabled == ECheckBoxState::Undetermined) && bControlBufferSize;
 
-			NewestDataTime = TimedDataMonitorSubsystem->GetInputNewestDataTime(InputIdentifer);
+			NewestDataTime = TimedDataMonitorSubsystem->GetInputNewestDataTime(InputIdentifier);
 
 			CachedStatsBufferUnderflow = 0;
 			CachedStatsBufferOverflow = 0;
@@ -124,14 +131,16 @@ struct FTimedDataInputTableRowData : TSharedFromThis<FTimedDataInputTableRowData
 		else
 		{
 			CachedEnabled = TimedDataMonitorSubsystem->IsChannelEnabled(ChannelIdentifier) ? ECheckBoxState::Checked : ECheckBoxState::Unchecked;
-			CachedInputEvaluationType = TimedDataMonitorSubsystem->GetInputEvaluationType(InputIdentifer);
+			CachedInputEvaluationType = TimedDataMonitorSubsystem->GetInputEvaluationType(InputIdentifier);
 			CachedInputEvaluationOffset = 0.f;
-			CachedState = TimedDataMonitorSubsystem->GetChannelState(ChannelIdentifier);
-			CachedBufferSize = TimedDataMonitorSubsystem->GetChannelNumberOfSamples(ChannelIdentifier);
+			CachedState = TimedDataMonitorSubsystem->GetChannelConnectionState(ChannelIdentifier);
+			CachedBufferSize = TimedDataMonitorSubsystem->GetChannelDataBufferSize(ChannelIdentifier);
+			CachedCurrentAmountOfBuffer = TimedDataMonitorSubsystem->GetChannelNumberOfSamples(ChannelIdentifier);
 			CachedStatsBufferUnderflow = TimedDataMonitorSubsystem->GetChannelBufferUnderflowStat(ChannelIdentifier);
 			CachedStatsBufferOverflow = TimedDataMonitorSubsystem->GetChannelBufferOverflowStat(ChannelIdentifier);
 			CachedStatsFrameDropped = TimedDataMonitorSubsystem->GetChannelFrameDroppedStat(ChannelIdentifier);
-			bCachedCanEditBufferSize = (CachedEnabled == ECheckBoxState::Checked || CachedEnabled == ECheckBoxState::Undetermined) && !TimedDataMonitorSubsystem->IsDataBufferSizeControlledByInput(InputIdentifer);
+			bControlBufferSize = !TimedDataMonitorSubsystem->IsDataBufferSizeControlledByInput(InputIdentifier);
+			bCachedCanEditBufferSize = (CachedEnabled == ECheckBoxState::Checked || CachedEnabled == ECheckBoxState::Undetermined) && bControlBufferSize;
 
 			NewestDataTime = TimedDataMonitorSubsystem->GetChannelNewestDataTime(ChannelIdentifier);
 		}
@@ -161,7 +170,7 @@ struct FTimedDataInputTableRowData : TSharedFromThis<FTimedDataInputTableRowData
 	}
 
 public:
-	FTimedDataMonitorInputIdentifier InputIdentifer;
+	FTimedDataMonitorInputIdentifier InputIdentifier;
 	FTimedDataMonitorChannelIdentifier ChannelIdentifier;
 	bool bIsInput;
 
@@ -175,9 +184,11 @@ public:
 	ETimedDataInputState CachedState = ETimedDataInputState::Disconnected;
 	FText CachedDescription;
 	int32 CachedBufferSize = 0;
+	int32 CachedCurrentAmountOfBuffer = 0;
 	int32 CachedStatsBufferUnderflow = 0;
 	int32 CachedStatsBufferOverflow = 0;
 	int32 CachedStatsFrameDropped = 0;
+	bool bControlBufferSize = false;
 	bool bCachedCanEditBufferSize = false;
 };
 
@@ -199,17 +210,26 @@ void STimedDataInputTableRow::Construct(const FArguments& InArgs, const TSharedR
 	}
 	else
 	{
-		Arg.Style(FTimedDataMonitorStyle::Get(), "TableView.Child");
+		Arg.Style(FTimedDataMonitorEditorStyle::Get(), "TableView.Child");
 	}
 	Super::Construct(Arg, InOwerTableView);
 }
 
 
+void STimedDataInputTableRow::UpdateCachedValue()
+{
+	if (DiagramWidget)
+	{
+		DiagramWidget->UpdateCachedValue();
+	}
+}
+
+
 TSharedRef<SWidget> STimedDataInputTableRow::GenerateWidgetForColumn(const FName& ColumnName)
 {
 	const FTextBlockStyle* ItemTextBlockStyle = Item->bIsInput
-		? &FTimedDataMonitorStyle::Get().GetWidgetStyle<FTextBlockStyle>("TextBlock.Large")
-		: &FTimedDataMonitorStyle::Get().GetWidgetStyle<FTextBlockStyle>("TextBlock.Regular");
+		? &FTimedDataMonitorEditorStyle::Get().GetWidgetStyle<FTextBlockStyle>("TextBlock.Large")
+		: &FTimedDataMonitorEditorStyle::Get().GetWidgetStyle<FTextBlockStyle>("TextBlock.Regular");
 
 	if (TimedDataListView::HeaderIdName_Enable == ColumnName)
 	{
@@ -217,7 +237,7 @@ TSharedRef<SWidget> STimedDataInputTableRow::GenerateWidgetForColumn(const FName
 			? LOCTEXT("ToggleAllChannelsToolTip", "Toggles all channels from this input.")
 			: LOCTEXT("ToggleChannelToolTip", "Toggles whether this channel will collect stats and be used when calibrating.");
 		return SNew(SCheckBox)
-			.Style(FTimedDataMonitorStyle::Get(), "CheckBox.Enable")
+			.Style(FTimedDataMonitorEditorStyle::Get(), "CheckBox.Enable")
 			.ToolTipText(Tooltip)
 			.IsChecked(this, &STimedDataInputTableRow::GetEnabledCheckState)
 			.OnCheckStateChanged(this, &STimedDataInputTableRow::OnEnabledCheckStateChanged);
@@ -236,6 +256,7 @@ TSharedRef<SWidget> STimedDataInputTableRow::GenerateWidgetForColumn(const FName
 					.IndentAmount(12)
 				]
 				+ SHorizontalBox::Slot()
+				.VAlign(VAlign_Center)
 				.FillWidth(1.0f)
 				[
 					SNew(SImage)
@@ -244,11 +265,28 @@ TSharedRef<SWidget> STimedDataInputTableRow::GenerateWidgetForColumn(const FName
 		}
 		return SNullWidget::NullWidget;
 	}
+	if (TimedDataListView::HeaderIdName_EvaluationMode == ColumnName)
+	{
+		if (Item->bIsInput)
+		{
+			return SNew(SComboButton)
+				.ComboButtonStyle(FTimedDataMonitorEditorStyle::Get(), "ToggleComboButton")
+				.HAlign(HAlign_Center)
+				.OnGetMenuContent(this, &STimedDataInputTableRow::OnEvaluationImageBuildMenu)
+				.ButtonContent()
+				[
+					SNew(SImage)
+					.Image(this, &STimedDataInputTableRow::GetEvaluationImage)
+				];
+		}
+		return SNullWidget::NullWidget;
+	}
 	if (TimedDataListView::HeaderIdName_Name == ColumnName)
 	{
 		return SNew(SHorizontalBox)
 			+ SHorizontalBox::Slot()
 			.Padding(10, 0, 10, 0)
+			.VAlign(VAlign_Center)
 			.AutoWidth()
 			[
 				SNew(STextBlock)
@@ -258,6 +296,7 @@ TSharedRef<SWidget> STimedDataInputTableRow::GenerateWidgetForColumn(const FName
 			]
 			+ SHorizontalBox::Slot()
 			.HAlign(HAlign_Left)
+			.VAlign(VAlign_Center)
 			[
 				SNew(STextBlock)
 				.Text(Item->DisplayName)
@@ -266,64 +305,105 @@ TSharedRef<SWidget> STimedDataInputTableRow::GenerateWidgetForColumn(const FName
 	}
 	if (TimedDataListView::HeaderIdName_Description == ColumnName)
 	{
-		return SNew(STextBlock)
-			.Text(this, &STimedDataInputTableRow::GetDescription)
-			.TextStyle(ItemTextBlockStyle);
+		return SNew(SHorizontalBox)
+			+ SHorizontalBox::Slot()
+			.VAlign(VAlign_Center)
+			.AutoWidth()
+			[
+				SNew(STextBlock)
+				.Text(this, &STimedDataInputTableRow::GetDescription)
+				.TextStyle(ItemTextBlockStyle)
+			];
 	}
 	if (TimedDataListView::HeaderIdName_TimeCorrection == ColumnName)
 	{
 		if (Item->bIsInput)
 		{
-			return SNew(STextBlock)
+			return SNew(STimedDataNumericEntryBox<float>)
 				.TextStyle(ItemTextBlockStyle)
-				.Text(this, &STimedDataInputTableRow::GetEvaluationOffsetText);
+				.MinValue(1)
+				.ToolTipText(LOCTEXT("TimeCorrection_ToolTip", "Time Correction."))
+				.Value(this, &STimedDataInputTableRow::GetEvaluationOffset)
+				.EditLabel(LOCTEXT("TimeCorrection_EditLable", "In Seconds: "))
+				.OnValueCommitted(this, &STimedDataInputTableRow::SetEvaluationOffset);
 		}
 		return SNullWidget::NullWidget;
 	}
 	if (TimedDataListView::HeaderIdName_BufferSize == ColumnName)
 	{
-		//@todo put proper editing widget
-		if (Item->bIsInput) // bCachedCanEditBufferSize
+		if (Item->bControlBufferSize)
 		{
-			return SNew(SNumericEntryBox<int32>)
-				.ToolTipText(LOCTEXT("BufferSize_ToolTip", "Buffer Size."))
+			return SNew(STimedDataNumericEntryBox<int32>)
+				.TextStyle(ItemTextBlockStyle)
 				.MinValue(1)
-				.MinDesiredValueWidth(50)
+				.ToolTipText(LOCTEXT("BufferSize_ToolTip", "Buffer Size."))
 				.Value(this, &STimedDataInputTableRow::GetBufferSize)
+				.ShowAmount(!Item->bIsInput)
+				.Amount(this, &STimedDataInputTableRow::GetCurrentSampleCount)
+				.EditLabel(LOCTEXT("BufferSize_EditLable", "Number of buffer: "))
 				.OnValueCommitted(this, &STimedDataInputTableRow::SetBufferSize)
+				.CanEdit(Item->bControlBufferSize)
 				.IsEnabled(this, &STimedDataInputTableRow::CanEditBufferSize);
 		}
 		else
 		{
-			return SNew(STextBlock)
-				.TextStyle(ItemTextBlockStyle)
-				.Text(this, &STimedDataInputTableRow::GetBufferSizeText);
+			return SNew(SHorizontalBox)
+				+ SHorizontalBox::Slot()
+				.VAlign(VAlign_Center)
+				.Padding(4.0f, 0.0f, 0.0f, 0.0f)
+				[
+					SNew(STextBlock)
+					.TextStyle(ItemTextBlockStyle)
+					.Text(this, &STimedDataInputTableRow::GetBufferSizeText)
+				];
 		}
 	}
 	if (TimedDataListView::HeaderIdName_BufferUnder == ColumnName)
 	{
-		return SNew(STextBlock)
-			.Text(this, &STimedDataInputTableRow::GetBufferUnderflowCount)
-			.TextStyle(ItemTextBlockStyle);
+		return SNew(SHorizontalBox)
+			+ SHorizontalBox::Slot()
+			.HAlign(HAlign_Right)
+			.VAlign(VAlign_Center)
+			.Padding(4.0f, 0.0f, 0.0f, 0.0f)
+			[
+				SNew(STextBlock)
+				.Text(this, &STimedDataInputTableRow::GetBufferUnderflowCount)
+				.TextStyle(ItemTextBlockStyle)
+			];
 	}
 	if (TimedDataListView::HeaderIdName_BufferOver == ColumnName)
 	{
-		return SNew(STextBlock)
-			.Text(this, &STimedDataInputTableRow::GetBufferOverflowCount)
-			.TextStyle(ItemTextBlockStyle);
+		return SNew(SHorizontalBox)
+			+ SHorizontalBox::Slot()
+			.HAlign(HAlign_Right)
+			.VAlign(VAlign_Center)
+			.Padding(4.0f, 0.0f, 0.0f, 0.0f)
+			[
+				SNew(STextBlock)
+				.Text(this, &STimedDataInputTableRow::GetBufferOverflowCount)
+				.TextStyle(ItemTextBlockStyle)
+			];
 	}
 	if (TimedDataListView::HeaderIdName_FrameDrop == ColumnName)
 	{
-		return SNew(STextBlock)
-			.Text(this, &STimedDataInputTableRow::GetFrameDroppedCount)
-			.TextStyle(ItemTextBlockStyle);
+		return SNew(SHorizontalBox)
+			+ SHorizontalBox::Slot()
+			.HAlign(HAlign_Right)
+			.VAlign(VAlign_Center)
+			.Padding(4.0f, 0.0f, 0.0f, 0.0f)
+			[
+				SNew(STextBlock)
+				.Text(this, &STimedDataInputTableRow::GetFrameDroppedCount)
+				.TextStyle(ItemTextBlockStyle)
+			];
 	}
 	if (TimedDataListView::HeaderIdName_TimingDiagram == ColumnName)
 	{
-		//@todo put proper timing diagram
-		return SNew(SBorder)
-			.BorderImage(FTimedDataMonitorStyle::Get().GetBrush("Brush.White"))
-			.BorderBackgroundColor(FLinearColor::Green);
+		return SAssignNew(DiagramWidget, STimingDiagramWidget, Item->bIsInput)
+				.ChannelIdentifier(Item->ChannelIdentifier)
+				.InputIdentifier(Item->InputIdentifier)
+				.ShowSigma(true)
+				.ShowFurther(true);
 	}
 
 	return SNullWidget::NullWidget;
@@ -343,7 +423,7 @@ void STimedDataInputTableRow::OnEnabledCheckStateChanged(ECheckBoxState NewState
 
 	if (Item->bIsInput)
 	{
-		TimedDataMonitorSubsystem->SetInputEnabled(Item->InputIdentifer, NewState == ECheckBoxState::Checked);
+		TimedDataMonitorSubsystem->SetInputEnabled(Item->InputIdentifier, NewState == ECheckBoxState::Checked);
 	}
 	else
 	{
@@ -394,7 +474,26 @@ FText STimedDataInputTableRow::GetEvaluationOffsetText() const
 }
 
 
-TOptional<int32> STimedDataInputTableRow::GetBufferSize() const
+float STimedDataInputTableRow::GetEvaluationOffset() const
+{
+	return Item->CachedInputEvaluationOffset;
+}
+
+
+void STimedDataInputTableRow::SetEvaluationOffset(float InValue, ETextCommit::Type CommitType)
+{
+	if (GetEvaluationOffset() != InValue && Item->bIsInput)
+	{
+		UTimedDataMonitorSubsystem* TimedDataMonitorSubsystem = GEngine->GetEngineSubsystem<UTimedDataMonitorSubsystem>();
+		check(TimedDataMonitorSubsystem);
+
+		TimedDataMonitorSubsystem->SetInputEvaluationOffsetInSeconds(Item->InputIdentifier, InValue);
+		OwnerTreeView->RequestRefresh();
+	}
+}
+
+
+int32 STimedDataInputTableRow::GetBufferSize() const
 {
 	return Item->CachedBufferSize;
 }
@@ -402,24 +501,33 @@ TOptional<int32> STimedDataInputTableRow::GetBufferSize() const
 
 FText STimedDataInputTableRow::GetBufferSizeText() const
 {
-	return FText::AsNumber(Item->CachedBufferSize);
+	if (Item->bIsInput)
+	{
+		return FText::GetEmpty();
+	}
+	else
+	{
+		return FText::Format(LOCTEXT("ChannelBufferSizeFormat", "{0}/{1}"), Item->CachedCurrentAmountOfBuffer, Item->CachedBufferSize);
+	}
 }
 
 
 void STimedDataInputTableRow::SetBufferSize(int32 InValue, ETextCommit::Type InType)
 {
-	if (InType == ETextCommit::OnEnter || InType == ETextCommit::OnUserMovedFocus)
+	if (GetBufferSize() != InValue)
 	{
 		UTimedDataMonitorSubsystem* TimedDataMonitorSubsystem = GEngine->GetEngineSubsystem<UTimedDataMonitorSubsystem>();
 		check(TimedDataMonitorSubsystem);
 
 		if (Item->bIsInput)
 		{
-			TimedDataMonitorSubsystem->SetInputDataBufferSize(Item->InputIdentifer, InValue);
+			TimedDataMonitorSubsystem->SetInputDataBufferSize(Item->InputIdentifier, InValue);
+			Item->CachedBufferSize = TimedDataMonitorSubsystem->GetInputDataBufferSize(Item->InputIdentifier);
 		}
 		else
 		{
 			TimedDataMonitorSubsystem->SetChannelDataBufferSize(Item->ChannelIdentifier, InValue);
+			Item->CachedBufferSize = TimedDataMonitorSubsystem->GetChannelDataBufferSize(Item->ChannelIdentifier);
 		}
 		OwnerTreeView->RequestRefresh();
 	}
@@ -432,6 +540,82 @@ bool STimedDataInputTableRow::CanEditBufferSize() const
 }
 
 
+int32 STimedDataInputTableRow::GetCurrentSampleCount() const
+{
+	return Item->CachedCurrentAmountOfBuffer;
+}
+
+TSharedRef<SWidget> STimedDataInputTableRow::OnEvaluationImageBuildMenu()
+{
+	FMenuBuilder MenuBuilder(true, nullptr);
+
+	const ETimedDataInputEvaluationType CurrentEvaluationType = Item->CachedInputEvaluationType;
+
+	ETimedDataInputEvaluationType LambdaEvaluationType = ETimedDataInputEvaluationType::Timecode;
+	MenuBuilder.AddMenuEntry(
+		LOCTEXT("EvaluationTypeTimecodeLabel", "Timecode"),
+		LOCTEXT("EvaluationTypeTimecodeTooltip", "Evaluate the input base on the engine's timecode value."),
+		FSlateIcon(FTimedDataMonitorEditorStyle::Get().GetStyleSetName(), FTimedDataMonitorEditorStyle::NAME_TimecodeBrush),
+		FUIAction(
+			FExecuteAction::CreateSP(this, &STimedDataInputTableRow::SetInputEvaluationType, LambdaEvaluationType),
+			FCanExecuteAction(),
+			FIsActionChecked::CreateLambda([CurrentEvaluationType, LambdaEvaluationType]() { return CurrentEvaluationType == LambdaEvaluationType; })
+		));
+
+	LambdaEvaluationType = ETimedDataInputEvaluationType::PlatformTime;
+	MenuBuilder.AddMenuEntry(
+		LOCTEXT("EvaluationTypePlatformTimeLabel", "Platform Time"),
+		LOCTEXT("EvaluationTypePlatformTimeTooltip", "Evaluate the input base ont he engine's time."),
+		FSlateIcon(FTimedDataMonitorEditorStyle::Get().GetStyleSetName(), FTimedDataMonitorEditorStyle::NAME_PlatformTimeBrush),
+		FUIAction(
+			FExecuteAction::CreateSP(this, &STimedDataInputTableRow::SetInputEvaluationType, LambdaEvaluationType),
+			FCanExecuteAction(),
+			FIsActionChecked::CreateLambda([CurrentEvaluationType, LambdaEvaluationType]() { return CurrentEvaluationType == LambdaEvaluationType; })
+		));
+
+	LambdaEvaluationType = ETimedDataInputEvaluationType::None;
+	MenuBuilder.AddMenuEntry(
+		LOCTEXT("EvaluationTypeNoneLabel", "No synchronization"),
+		LOCTEXT("EvaluationTypeNoneTooltip", "Do not create any special evaluation (take the latest sample available)."),
+		FSlateIcon(FTimedDataMonitorEditorStyle::Get().GetStyleSetName(), FTimedDataMonitorEditorStyle::NAME_NoEvaluationBrush),
+		FUIAction(
+			FExecuteAction::CreateSP(this, &STimedDataInputTableRow::SetInputEvaluationType, LambdaEvaluationType),
+			FCanExecuteAction(),
+			FIsActionChecked::CreateLambda([CurrentEvaluationType, LambdaEvaluationType]() { return CurrentEvaluationType == LambdaEvaluationType; })
+		));
+
+	return MenuBuilder.MakeWidget();
+}
+
+
+const FSlateBrush* STimedDataInputTableRow::GetEvaluationImage() const
+{
+	if (Item->CachedInputEvaluationType == ETimedDataInputEvaluationType::Timecode)
+	{
+		return FTimedDataMonitorEditorStyle::Get().GetBrush(FTimedDataMonitorEditorStyle::NAME_TimecodeBrush);
+	}
+	if (Item->CachedInputEvaluationType == ETimedDataInputEvaluationType::PlatformTime)
+	{
+		return FTimedDataMonitorEditorStyle::Get().GetBrush(FTimedDataMonitorEditorStyle::NAME_PlatformTimeBrush);
+	}
+	return FTimedDataMonitorEditorStyle::Get().GetBrush(FTimedDataMonitorEditorStyle::NAME_NoEvaluationBrush);
+}
+
+
+void STimedDataInputTableRow::SetInputEvaluationType(ETimedDataInputEvaluationType EvaluationType)
+{
+	UTimedDataMonitorSubsystem* TimedDataMonitorSubsystem = GEngine->GetEngineSubsystem<UTimedDataMonitorSubsystem>();
+	check(TimedDataMonitorSubsystem);
+
+	if (Item->bIsInput)
+	{
+		TimedDataMonitorSubsystem->SetInputEvaluationType(Item->InputIdentifier, EvaluationType);
+
+		OwnerTreeView->RequestRefresh();
+	}
+}
+
+
 FText STimedDataInputTableRow::GetBufferUnderflowCount() const
 {
 	return FText::AsNumber(Item->CachedStatsBufferUnderflow);
@@ -450,8 +634,9 @@ FText STimedDataInputTableRow::GetFrameDroppedCount() const
 /**
  * STimedDataListView
  */
-void STimedDataInputListView::Construct(const FArguments& InArgs)
+void STimedDataInputListView::Construct(const FArguments& InArgs, TSharedPtr<STimedDataMonitorPanel> InOwnerPanel)
 {
+	OwnerPanel = InOwnerPanel;
 	UTimedDataMonitorSubsystem* TimedDataMonitorSubsystem = GEngine->GetEngineSubsystem<UTimedDataMonitorSubsystem>();
 	check(TimedDataMonitorSubsystem);
 	TimedDataMonitorSubsystem->OnIdentifierListChanged().AddSP(this, &STimedDataInputListView::RequestRebuildSources);
@@ -462,9 +647,11 @@ void STimedDataInputListView::Construct(const FArguments& InArgs)
 		.TreeItemsSource(&ListItemsSource)
 		.SelectionMode(ESelectionMode::SingleToggle)
 		.OnGenerateRow(this, &STimedDataInputListView::OnGenerateRow)
+		.OnRowReleased(this, &STimedDataInputListView::ReleaseListViewWidget)
 		.OnGetChildren(this, &STimedDataInputListView::GetChildrenForInfo)
 		.OnSelectionChanged(this, &STimedDataInputListView::OnSelectionChanged)
 		.OnIsSelectableOrNavigable(this, &STimedDataInputListView::OnIsSelectableOrNavigable)
+		.OnContextMenuOpening(InArgs._OnContextMenuOpening)
 		.HighlightParentNodesForSelection(true)
 		.HeaderRow
 		(
@@ -487,43 +674,52 @@ void STimedDataInputListView::Construct(const FArguments& InArgs)
 			.DefaultLabel(FText::GetEmpty())
 
 			+ SHeaderRow::Column(TimedDataListView::HeaderIdName_Name)
-			.FillWidth(0.33)
+			.FillWidth(0.33f)
 			.HAlignCell(EHorizontalAlignment::HAlign_Left)
 			.DefaultLabel(LOCTEXT("HeaderName_Name", "Name"))
 
 			+ SHeaderRow::Column(TimedDataListView::HeaderIdName_Description)
-			.FillWidth(0.33)
+			.FillWidth(0.33f)
+			.HAlignCell(EHorizontalAlignment::HAlign_Left)
+			.DefaultLabel(LOCTEXT("HeaderName_Description", "Last Sample Time"))
+
+			+ SHeaderRow::Column(TimedDataListView::HeaderIdName_EvaluationMode)
+			.FixedWidth(48)
 			.HAlignCell(EHorizontalAlignment::HAlign_Left)
-			.DefaultLabel(LOCTEXT("HeaderName_Description", "Description"))
+			.DefaultLabel(LOCTEXT("HeaderName_EvaluationMode", "Eval."))
+			.DefaultTooltip(LOCTEXT("HeaderTooltip_EvaluationMode", "How the input is evaluated"))
 
 			+ SHeaderRow::Column(TimedDataListView::HeaderIdName_TimeCorrection)
 			.FixedWidth(100)
-			.HAlignCell(EHorizontalAlignment::HAlign_Left)
+			.HAlignCell(EHorizontalAlignment::HAlign_Right)
 			.DefaultLabel(LOCTEXT("HeaderName_TimeCorrection", "Time Correction"))
 
 			+ SHeaderRow::Column(TimedDataListView::HeaderIdName_BufferSize)
 			.FixedWidth(100)
-			.HAlignCell(EHorizontalAlignment::HAlign_Left)
+			.HAlignCell(EHorizontalAlignment::HAlign_Right)
 			.DefaultLabel(LOCTEXT("HeaderName_BufferSize", "Buffer Size"))
 
 			+ SHeaderRow::Column(TimedDataListView::HeaderIdName_BufferUnder)
 			.FixedWidth(50)
-			.HAlignCell(EHorizontalAlignment::HAlign_Left)
+			.HAlignCell(EHorizontalAlignment::HAlign_Right)
 			.DefaultLabel(LOCTEXT("HeaderName_BufferUnder", "B.U."))
+			.DefaultTooltip(LOCTEXT("HeaderTooltip_BufferUnder", "Number of buffer underflows detected"))
 
 			+ SHeaderRow::Column(TimedDataListView::HeaderIdName_BufferOver)
 			.FixedWidth(50)
-			.HAlignCell(EHorizontalAlignment::HAlign_Left)
+			.HAlignCell(EHorizontalAlignment::HAlign_Right)
 			.DefaultLabel(LOCTEXT("HeaderName_BufferOver", "B.O."))
+			.DefaultTooltip(LOCTEXT("HeaderTooltip_BufferOver", "Number of buffer overflows detected"))
 
 			+ SHeaderRow::Column(TimedDataListView::HeaderIdName_FrameDrop)
 			.FixedWidth(50)
-			.HAlignCell(EHorizontalAlignment::HAlign_Left)
+			.HAlignCell(EHorizontalAlignment::HAlign_Right)
 			.DefaultLabel(LOCTEXT("HeaderName_FrameDrop", "F.D."))
+			.DefaultTooltip(LOCTEXT("HeaderTooltip_FrameDrop", "Number of frame drop detected"))
 
 			+ SHeaderRow::Column(TimedDataListView::HeaderIdName_TimingDiagram)
-			.FillWidth(0.33)
-			.HAlignCell(EHorizontalAlignment::HAlign_Left)
+			.FillWidth(0.33f)
+			.HAlignCell(EHorizontalAlignment::HAlign_Fill)
 			.DefaultLabel(LOCTEXT("HeaderName_TimingDiagram", "Timing Diagram"))
 		)
 	);
@@ -540,10 +736,17 @@ STimedDataInputListView::~STimedDataInputListView()
 }
 
 
-void STimedDataInputListView::Tick(const FGeometry& AllottedGeometry, const double InCurrentTime, const float InDeltaTime)
+void STimedDataInputListView::RequestRefresh()
 {
-	Super::Tick(AllottedGeometry, InCurrentTime, InDeltaTime);
+	if (TSharedPtr<STimedDataMonitorPanel> OwnerPanelPin = OwnerPanel.Pin())
+	{
+		OwnerPanelPin->RequestRefresh();
+	}
+}
 
+
+void STimedDataInputListView::UpdateCachedValue()
+{
 	if (bRebuildListRequested)
 	{
 		RebuildSources();
@@ -551,14 +754,21 @@ void STimedDataInputListView::Tick(const FGeometry& AllottedGeometry, const doub
 		bRebuildListRequested = false;
 	}
 
-	double RefreshTimer = GetDefault<UTimedDataMonitorEditorSettings>()->RefreshRate;
-	if (bRefreshRequested || (FApp::GetCurrentTime() - LastCachedValueUpdateTime > RefreshTimer))
+	for (FTimedDataInputTableRowDataPtr& RowDataPtr : ListItemsSource)
+	{
+		RowDataPtr->UpdateCachedValue();
+	}
+
+	for (int32 Index = ListRowWidgets.Num() - 1; Index >= 0; --Index)
 	{
-		bRefreshRequested = false;
-		LastCachedValueUpdateTime = FApp::GetCurrentTime();
-		for (FTimedDataInputTableRowDataPtr& RowDataPtr : ListItemsSource)
+		const TSharedPtr<STimedDataInputTableRow> Row = ListRowWidgets[Index].Pin();
+		if (Row)
+		{
+			Row->UpdateCachedValue();
+		}
+		else
 		{
-			RowDataPtr->UpdateCachedValue();
+			ListRowWidgets.RemoveAtSwap(Index);
 		}
 	}
 }
@@ -600,6 +810,17 @@ void STimedDataInputListView::RebuildSources()
 }
 
 
+FTimedDataMonitorInputIdentifier STimedDataInputListView::GetSelectedInputIdentifier() const
+{
+	TArray<FTimedDataInputTableRowDataPtr> SelectedRows = GetSelectedItems();
+	if (SelectedRows.Num() > 0)
+	{
+		return SelectedRows[0]->InputIdentifier;
+	}
+	return FTimedDataMonitorInputIdentifier();
+}
+
+
 ECheckBoxState STimedDataInputListView::GetAllEnabledCheckState() const
 {
 	return ECheckBoxState::Checked;
@@ -614,18 +835,29 @@ void STimedDataInputListView::OnToggleAllEnabledCheckState(ECheckBoxState CheckB
 	bool bIsEnabled = CheckBoxState == ECheckBoxState::Checked;
 	for (const FTimedDataInputTableRowDataPtr& RowDataPtr : ListItemsSource)
 	{
-		TimedDataMonitorSubsystem->SetInputEnabled(RowDataPtr->InputIdentifer, bIsEnabled);
+		TimedDataMonitorSubsystem->SetInputEnabled(RowDataPtr->InputIdentifier, bIsEnabled);
 	}
 }
 
 
 TSharedRef<ITableRow> STimedDataInputListView::OnGenerateRow(FTimedDataInputTableRowDataPtr InItem, const TSharedRef<STableViewBase>& OwnerTable)
 {
-	return SNew(STimedDataInputTableRow, OwnerTable, SharedThis<STimedDataInputListView>(this))
+	TSharedRef<STimedDataInputTableRow> Row = SNew(STimedDataInputTableRow, OwnerTable, SharedThis<STimedDataInputListView>(this))
 		.Item(InItem);
+	ListRowWidgets.Add(Row);
+	return Row;
+}
+
+
+void STimedDataInputListView::ReleaseListViewWidget(const TSharedRef<ITableRow>& Row)
+{
+	TSharedRef<STimedDataInputTableRow> RefRow = StaticCastSharedRef<STimedDataInputTableRow>(Row);
+	TWeakPtr<STimedDataInputTableRow> WeakRow = RefRow;
+	ListRowWidgets.RemoveSingleSwap(WeakRow);
 }
 
 
+
 void STimedDataInputListView::GetChildrenForInfo(FTimedDataInputTableRowDataPtr InItem, TArray<FTimedDataInputTableRowDataPtr>& OutChildren)
 {
 	OutChildren = InItem->InputChildren;
@@ -636,7 +868,7 @@ void STimedDataInputListView::OnSelectionChanged(FTimedDataInputTableRowDataPtr
 {
 	if (SelectInfo != ESelectInfo::Direct)
 	{
-		if (InItem && InItem->bIsInput)
+		if (InItem && !InItem->bIsInput)
 		{
 			ClearSelection();
 		}
@@ -649,4 +881,5 @@ bool STimedDataInputListView::OnIsSelectableOrNavigable(FTimedDataInputTableRowD
 	return InItem && !InItem->bIsInput;
 }
 
+
 #undef LOCTEXT_NAMESPACE
\ No newline at end of file
diff --git a/Engine/Plugins/VirtualProduction/TimedDataMonitor/Source/TimedDataMonitorEditor/Private/STimedDataListView.h b/Engine/Plugins/VirtualProduction/TimedDataMonitor/Source/TimedDataMonitorEditor/Private/STimedDataListView.h
--- a/Engine/Plugins/VirtualProduction/TimedDataMonitor/Source/TimedDataMonitorEditor/Private/STimedDataListView.h
+++ b/Engine/Plugins/VirtualProduction/TimedDataMonitor/Source/TimedDataMonitorEditor/Private/STimedDataListView.h
@@ -6,16 +6,20 @@
 
 #include "SlateFwd.h"
 #include "SlateOptMacros.h"
+#include "TimedDataMonitorTypes.h"
 #include "Widgets/DeclarativeSyntaxSupport.h"
 #include "Widgets/SCompoundWidget.h"
 
 #include "Widgets/Views/STreeView.h"
 #include "Widgets/Views/STableRow.h"
 
-
+enum class ETimedDataInputEvaluationType : uint8;
+struct FTimedDataInputTableRowData;
+struct FSlateBrush;
 class STimedDataInputListView;
 class STimedDataInputTableRow;
-struct FTimedDataInputTableRowData;
+class STimedDataMonitorPanel;
+class STimingDiagramWidget;
 
 
 /**
@@ -38,6 +42,9 @@ public:
 
 	void Construct(const FArguments& Args, const TSharedRef<STableViewBase>& OwerTableView, const TSharedRef<STimedDataInputListView>& OwnerTreeView);
 
+public:
+	void UpdateCachedValue();
+
 private:
 	virtual TSharedRef<SWidget> GenerateWidgetForColumn(const FName& ColumnName) override;
 
@@ -48,11 +55,17 @@ private:
 	FSlateColor GetStateColorAndOpacity() const;
 
 	FText GetDescription() const;
+	float GetEvaluationOffset() const;
 	FText GetEvaluationOffsetText() const;
-	TOptional<int32> GetBufferSize() const;
+	void SetEvaluationOffset(float NewValue, ETextCommit::Type CommitType);
+	int32 GetBufferSize() const;
 	FText GetBufferSizeText() const;
 	void SetBufferSize(int32 NewValue, ETextCommit::Type CommitType);
 	bool CanEditBufferSize() const;
+	int32 GetCurrentSampleCount() const;
+	TSharedRef<SWidget> OnEvaluationImageBuildMenu();
+	const FSlateBrush* GetEvaluationImage() const;
+	void SetInputEvaluationType(ETimedDataInputEvaluationType EvaluationType);
 
 	/** Queries about buffer stats to display */
 	FText GetBufferUnderflowCount() const;
@@ -62,6 +75,7 @@ private:
 private:
 	FTimedDataInputTableRowDataPtr Item;
 	TSharedPtr<STimedDataInputListView> OwnerTreeView;
+	TSharedPtr<STimingDiagramWidget> DiagramWidget;
 };
 
 
@@ -74,14 +88,16 @@ class STimedDataInputListView : public STreeView<FTimedDataInputTableRowDataPtr>
 
 public:
 	SLATE_BEGIN_ARGS(STimedDataInputListView) {}
+		SLATE_EVENT(FOnContextMenuOpening, OnContextMenuOpening)
 	SLATE_END_ARGS()
 
-	void Construct(const FArguments& InArgs);
+	void Construct(const FArguments& InArgs, TSharedPtr<STimedDataMonitorPanel> OwnerPanel);
 	virtual ~STimedDataInputListView();
 
-	virtual void Tick(const FGeometry& AllottedGeometry, const double InCurrentTime, const float InDeltaTime) override;
+	void RequestRefresh();
+	void UpdateCachedValue();
 
-	void RequestRefresh() { bRefreshRequested = true; }
+	FTimedDataMonitorInputIdentifier GetSelectedInputIdentifier() const;
 
 private:
 	void RequestRebuildSources();
@@ -90,13 +106,16 @@ private:
 	ECheckBoxState GetAllEnabledCheckState() const;
 	void OnToggleAllEnabledCheckState(ECheckBoxState CheckBoxState);
 	TSharedRef<ITableRow> OnGenerateRow(FTimedDataInputTableRowDataPtr Item, const TSharedRef<STableViewBase>& OwnerTable);
+	void ReleaseListViewWidget(const TSharedRef<ITableRow>& Row);
 	void GetChildrenForInfo(FTimedDataInputTableRowDataPtr InItem, TArray<FTimedDataInputTableRowDataPtr>& OutChildren);
 	void OnSelectionChanged(FTimedDataInputTableRowDataPtr InItem, ESelectInfo::Type SelectInfo);
 	bool OnIsSelectableOrNavigable(FTimedDataInputTableRowDataPtr InItem) const;
 
+
 private:
+	TWeakPtr<STimedDataMonitorPanel> OwnerPanel;
+
 	TArray<FTimedDataInputTableRowDataPtr> ListItemsSource;
+	TArray<TWeakPtr<STimedDataInputTableRow>> ListRowWidgets;
 	bool bRebuildListRequested = true;
-	bool bRefreshRequested = true;
-	double LastCachedValueUpdateTime = 0.0;
 };
