- Updating MediaIO to work with TimedData monitor

diff --git a/Engine/Plugins/VirtualProduction/TimedDataMonitor/Source/TimedDataMonitorEditor/Private/STimedDataMonitorPanel.cpp b/Engine/Plugins/VirtualProduction/TimedDataMonitor/Source/TimedDataMonitorEditor/Private/STimedDataMonitorPanel.cpp
--- a/Engine/Plugins/VirtualProduction/TimedDataMonitor/Source/TimedDataMonitorEditor/Private/STimedDataMonitorPanel.cpp
+++ b/Engine/Plugins/VirtualProduction/TimedDataMonitor/Source/TimedDataMonitorEditor/Private/STimedDataMonitorPanel.cpp
@@ -3,10 +3,8 @@
 #include "STimedDataMonitorPanel.h"
 
 #include "Editor.h"
-#include "EditorFontGlyphs.h"
 #include "Editor/EditorPerformanceSettings.h"
 #include "Engine/Engine.h"
-#include "Framework/MultiBox/MultiBoxBuilder.h"
 #include "IMessageLogListing.h"
 #include "ISettingsModule.h"
 #include "ITimedDataInput.h"
@@ -15,18 +13,29 @@
 #include "Logging/MessageLog.h"
 #include "MessageLogModule.h"
 #include "TimedDataInputCollection.h"
+#include "TimedDataMonitorCalibration.h"
 #include "TimedDataMonitorSubsystem.h"
+#include "TimedDataMonitorEditorSettings.h"
+
+#include "EditorFontGlyphs.h"
+#include "TimedDataMonitorEditorStyle.h"
+
+#include "Framework/MultiBox/MultiBoxBuilder.h"
+#include "Framework/Docking/TabManager.h"
 
-#include "STimecodeProvider.h"
+#include "STimedDataGenlock.h"
+#include "STimedDataMonitorBufferVisualizer.h"
+#include "STimedDataTimecodeProvider.h"
 #include "STimedDataListView.h"
 #include "Widgets/Docking/SDockTab.h"
+#include "Widgets/Images/SImage.h"
+#include "Widgets/Images/SThrobber.h"
 #include "Widgets/Input/SButton.h"
 #include "Widgets/Input/SCheckBox.h"
 #include "Widgets/Layout/SBorder.h"
 #include "Widgets/Layout/SScrollBox.h"
 #include "Widgets/Layout/SSpacer.h"
-#include "Widgets/TimedDataMonitorStyle.h"
-
+#include "Widgets/SOverlay.h"
 
 #define LOCTEXT_NAMESPACE "STimedDataMonitorPanel"
 
@@ -38,10 +47,11 @@ FDelegateHandle STimedDataMonitorPanel::LevelEditorTabManagerChangedHandle;
 namespace Utilities
 {
 	static const FName NAME_App = FName("TimedDataSourceMonitorPanelApp");
+	static const FName NAME_TimedDataMonitorBuffers = FName("TimedDataMonitorBuffers");
 	static const FName NAME_LevelEditorModuleName("LevelEditor");
 	static const FName NAME_LogName = "Timed Data Monitor";
 
-	TSharedRef<SDockTab> CreateTab(const FSpawnTabArgs& Args)
+	TSharedRef<SDockTab> CreateApp(const FSpawnTabArgs& Args)
 	{
 		return SNew(SDockTab)
 			.TabRole(ETabRole::NomadTab)
@@ -49,6 +59,15 @@ namespace Utilities
 				SNew(STimedDataMonitorPanel)
 			];
 	}
+
+	TSharedRef<SDockTab> CreateBuffersVisualizer(const FSpawnTabArgs&)
+	{
+		return SNew(SDockTab)
+			.TabRole(ETabRole::NomadTab)
+			[
+				SNew(STimedDataMonitorBufferVisualizer)
+			];
+	}
 }
 
 
@@ -59,11 +78,16 @@ void STimedDataMonitorPanel::RegisterNomadTabSpawner(TSharedRef<FWorkspaceItem>
 		FLevelEditorModule& LevelEditorModule = FModuleManager::LoadModuleChecked<FLevelEditorModule>(Utilities::NAME_LevelEditorModuleName);
 		TSharedPtr<FTabManager> LevelEditorTabManager = LevelEditorModule.GetLevelEditorTabManager();
 
-		LevelEditorTabManager->RegisterTabSpawner(Utilities::NAME_App, FOnSpawnTab::CreateStatic(&Utilities::CreateTab))
+		LevelEditorTabManager->RegisterTabSpawner(Utilities::NAME_App, FOnSpawnTab::CreateStatic(&Utilities::CreateApp))
 			.SetDisplayName(LOCTEXT("TabTitle", "Timed Input Monitor"))
 			.SetTooltipText(LOCTEXT("TooltipText", "Monitor inputs that can be time synchronized"))
 			.SetGroup(InWorkspaceItem)
-			.SetIcon(FSlateIcon(FTimedDataMonitorStyle::Get().GetStyleSetName(), "Img.Timecode.Small"));
+			.SetIcon(FSlateIcon(FTimedDataMonitorEditorStyle::Get().GetStyleSetName(), "Img.TimedDataMonitor.Small"));
+
+		FGlobalTabmanager::Get()->RegisterNomadTabSpawner(Utilities::NAME_TimedDataMonitorBuffers, FOnSpawnTab::CreateStatic(&Utilities::CreateBuffersVisualizer))
+			.SetDisplayName(LOCTEXT("BufferVisualizerTitle", "Buffer Visualizer"))
+			.SetTooltipText(LOCTEXT("BufferVisualizerTooltip", "Open buffer visualizer tab for Timed Data Monitor."))
+			.SetMenuType(ETabSpawnerMenuType::Hidden);
 	};
 
 	FLevelEditorModule& LevelEditorModule = FModuleManager::LoadModuleChecked<FLevelEditorModule>(Utilities::NAME_LevelEditorModuleName);
@@ -80,6 +104,8 @@ void STimedDataMonitorPanel::RegisterNomadTabSpawner(TSharedRef<FWorkspaceItem>
 
 void STimedDataMonitorPanel::UnregisterNomadTabSpawner()
 {
+	FGlobalTabmanager::Get()->UnregisterNomadTabSpawner(Utilities::NAME_TimedDataMonitorBuffers);
+
 	if (FSlateApplication::IsInitialized() && FModuleManager::Get().IsModuleLoaded(Utilities::NAME_LevelEditorModuleName))
 	{
 		FLevelEditorModule* LevelEditorModule = FModuleManager::GetModulePtr<FLevelEditorModule>(Utilities::NAME_LevelEditorModuleName);
@@ -121,204 +147,311 @@ void STimedDataMonitorPanel::Construct(const FArguments& InArgs)
 
 	BuildCalibrationArray();
 
-	TSharedRef<class SWidget> MessageLogListingWidget = MessageLogListing.IsValid() ? MessageLogModule.CreateLogListingWidget(MessageLogListing.ToSharedRef()) : SNullWidget::NullWidget;
+	TSharedRef<SWidget> MessageLogListingWidget = MessageLogListing.IsValid() ? MessageLogModule.CreateLogListingWidget(MessageLogListing.ToSharedRef()) : SNullWidget::NullWidget;
 
-	FProperty* PerformanceThrottlingProperty = FindFieldChecked<FProperty>(UEditorPerformanceSettings::StaticClass(), GET_MEMBER_NAME_CHECKED(UEditorPerformanceSettings, bThrottleCPUWhenNotForeground));
-	PerformanceThrottlingProperty->GetDisplayNameText();
-	FFormatNamedArguments Arguments;
-	Arguments.Add(TEXT("PropertyName"), PerformanceThrottlingProperty->GetDisplayNameText());
-	FText PerformanceWarningText = FText::Format(LOCTEXT("kPerformanceWarningMessage", "Warning: The editor setting '{PropertyName}' is currently enabled\nThis will stop editor windows from updating in realtime while the editor is not in focus"), Arguments);
-
-	const int32 ButtonBoxSize = 28;
+	TSharedPtr<SWidget> PerformanceWidget;
+	{
+		FProperty* PerformanceThrottlingProperty = FindFieldChecked<FProperty>(UEditorPerformanceSettings::StaticClass(), GET_MEMBER_NAME_CHECKED(UEditorPerformanceSettings, bThrottleCPUWhenNotForeground));
+		PerformanceThrottlingProperty->GetDisplayNameText();
+		FFormatNamedArguments Arguments;
+		Arguments.Add(TEXT("PropertyName"), PerformanceThrottlingProperty->GetDisplayNameText());
+		FText PerformanceWarningText = FText::Format(LOCTEXT("PerformanceWarningMessage", "Warning: The editor setting '{PropertyName}' is currently enabled\nThis will stop editor windows from updating in realtime while the editor is not in focus"), Arguments);
 
-	ChildSlot
-	[
-		SNew(SVerticalBox)
-		// toolbar
-		+ SVerticalBox::Slot()
-		.AutoHeight()
-		.Padding(10.f, 10.f, 10.f, 10.f)
-		[
-			SNew(SBorder)
-			.VAlign(VAlign_Center)
+		PerformanceWidget = SNew(SBorder)
+			.BorderImage(FEditorStyle::GetBrush("SettingsEditor.CheckoutWarningBorder"))
+			.BorderBackgroundColor(FColor(166, 137, 0))
 			.Padding(FMargin(5.f, 5.f, 5.f, 5.f))
+			.Visibility(this, &STimedDataMonitorPanel::ShowEditorPerformanceThrottlingWarning)
 			[
 				SNew(SHorizontalBox)
-				// Calibrate button
 				+ SHorizontalBox::Slot()
+				.FillWidth(1.0f)
+				.VAlign(VAlign_Center)
+				[
+					SNew(STextBlock)
+					.Text(PerformanceWarningText)
+					.Font(FEditorStyle::GetFontStyle("PropertyWindow.NormalFont"))
+					.ShadowColorAndOpacity(FLinearColor::Black.CopyWithNewOpacity(0.3f))
+					.ShadowOffset(FVector2D::UnitVector)
+				]
+				+ SHorizontalBox::Slot()
+				.Padding(FMargin(0.f, 0.f, 10.f, 0.f))
 				.AutoWidth()
 				.VAlign(VAlign_Center)
-				.Padding(0.f, 0.f, 10.f, 0.f)
 				[
-					SNew(SHorizontalBox)
-					+ SHorizontalBox::Slot()
-					.FillWidth(1.f)
+					SNew(SButton)
+					.OnClicked(this, &STimedDataMonitorPanel::DisableEditorPerformanceThrottling)
+					.Text(LOCTEXT("PerformanceWarningDisable", "Disable"))
+				]
+			];
+	}
+
+	TSharedRef<SWidget> Toolbar = SNew(SBorder)
+		.VAlign(VAlign_Center)
+		.Padding(FMargin(5.f, 5.f, 5.f, 5.f))
+		[
+			SNew(SHorizontalBox)
+			// Calibrate button
+			+ SHorizontalBox::Slot()
+			.AutoWidth()
+			.VAlign(VAlign_Center)
+			.Padding(4.f)
+			[
+				SNew(SHorizontalBox)
+				+ SHorizontalBox::Slot()
+				.FillWidth(1.f)
+				[
+					SNew(SButton)
+					.ToolTipText(this, &STimedDataMonitorPanel::GetCalibrateButtonTooltip)
+					.ButtonStyle(FTimedDataMonitorEditorStyle::Get(), "ToggleButton")
+					.ContentPadding(FMargin(4, 0))
+					.HAlign(HAlign_Center)
+					.VAlign(VAlign_Center)
+					.OnClicked(this, &STimedDataMonitorPanel::OnCalibrateClicked)
 					[
-						SNew(SButton)
-						.ToolTipText(this, &STimedDataMonitorPanel::GetCalibrateButtonTooltip)
-						.ButtonStyle(FEditorStyle::Get(), "FlatButton.Success")
-						.ForegroundColor(FLinearColor::White)
-						.ContentPadding(FMargin(4, 0))
-						.HAlign(HAlign_Center)
-						.VAlign(VAlign_Center)
-						.OnClicked(this, &STimedDataMonitorPanel::OnCalibrateClicked)
+						SNew(SHorizontalBox)
+						+ SHorizontalBox::Slot()
+						.AutoWidth()
+						[
+							SNew(SImage)
+							.Image(this, &STimedDataMonitorPanel::GetCalibrateButtonImage)
+						]
+						+ SHorizontalBox::Slot()
+						.AutoWidth()
+						.Padding(10.f, 4.f, 4.f, 4.f)
 						[
 							SNew(STextBlock)
-							.TextStyle(FTimedDataMonitorStyle::Get(), "TextBlock.Large")
+							.TextStyle(FTimedDataMonitorEditorStyle::Get(), "TextBlock.Large")
 							.Text(this, &STimedDataMonitorPanel::GetCalibrateButtonText)
+							.ColorAndOpacity(FLinearColor::White)
 						]
 					]
-					+ SHorizontalBox::Slot()
-					.AutoWidth()
-					[
-						SNew(SComboButton)
-						.ButtonStyle(FEditorStyle::Get(), "FlatButton.Success")
-						.ForegroundColor(FLinearColor::White)
-						.HAlign(HAlign_Center)
-						.VAlign(VAlign_Center)
-						.OnGetMenuContent(this, &STimedDataMonitorPanel::OnCalibrateBuildMenu)
-					]
 				]
-				// reset errors button
 				+ SHorizontalBox::Slot()
 				.AutoWidth()
-				.VAlign(VAlign_Center)
-				.HAlign(HAlign_Center)
-				.Padding(0.f)
+				[
+					SNew(SComboButton)
+					.ButtonStyle(FTimedDataMonitorEditorStyle::Get(), "ToggleButton")
+					.HAlign(HAlign_Center)
+					.VAlign(VAlign_Center)
+					.OnGetMenuContent(this, &STimedDataMonitorPanel::OnCalibrateBuildMenu)
+				]
+			]
+			// reset errors button
+			+ SHorizontalBox::Slot()
+			.AutoWidth()
+			.VAlign(VAlign_Center)
+			.Padding(4.f)
+			[
+				SNew(SHorizontalBox)
+				+ SHorizontalBox::Slot()
+				.FillWidth(1.f)
 				[
 					SNew(SButton)
-					.ToolTipText(LOCTEXT("ResetErrors_ToolTip", "Reset Errors"))
-					.ContentPadding(FMargin(4, 0))
+					.ButtonStyle(FTimedDataMonitorEditorStyle::Get(), "ToggleButton")
+					.ToolTipText(LOCTEXT("ResetErrors_ToolTip", "Reset the buffer statistics, message or evaluation time (base on settings)."))
+					.ContentPadding(FMargin(4, 4))
 					.HAlign(HAlign_Center)
 					.VAlign(VAlign_Center)
 					.OnClicked(this, &STimedDataMonitorPanel::OnResetErrorsClicked)
 					[
 						SNew(STextBlock)
-						.TextStyle(FTimedDataMonitorStyle::Get(), "TextBlock.Regular")
-						.Text(LOCTEXT("ResetErrors", "Reset Errors"))
+						.Font(FEditorStyle::Get().GetFontStyle("FontAwesome.14"))
+						.Text(FEditorFontGlyphs::Eraser)
+						.ColorAndOpacity(FLinearColor::White)
 					]
 				]
-				// Spacer
 				+ SHorizontalBox::Slot()
-				[
-					SNew(SSpacer)
-				]
-				// Settings button
-				+ SHorizontalBox::Slot()
-				.Padding(.0f)
 				.AutoWidth()
-				.HAlign(HAlign_Right)
-				.VAlign(VAlign_Center)
 				[
-					SNew(SBox)
-					.WidthOverride(ButtonBoxSize)
-					.HeightOverride(ButtonBoxSize)
-					[
-						SNew(SCheckBox)
-						.Padding(4.f)
-						.ToolTipText(LOCTEXT("ShowUserSettings_Tip", "Show/Hide the general user settings"))
-						.Style(FEditorStyle::Get(), "ToggleButtonCheckbox")
-						.ForegroundColor(FSlateColor::UseForeground())
-						.IsChecked_Lambda([]() { return ECheckBoxState::Unchecked; })
-						.OnCheckStateChanged_Lambda([](ECheckBoxState CheckState) { FModuleManager::LoadModuleChecked<ISettingsModule>("Settings").ShowViewer("Editor", "Plugins", "Timed Data Monitor"); })
-						[
-							SNew(STextBlock)
-							.Font(FEditorStyle::Get().GetFontStyle("FontAwesome.14"))
-							.Text(FEditorFontGlyphs::Cogs)
-						]
-					]
+					SNew(SComboButton)
+					.ButtonStyle(FTimedDataMonitorEditorStyle::Get(), "ToggleButton")
+					.HAlign(HAlign_Center)
+					.VAlign(VAlign_Center)
+					.OnGetMenuContent(this, &STimedDataMonitorPanel::OnResetBuildMenu)
 				]
 			]
-		]
-		// timing element
-		+ SVerticalBox::Slot()
-		.AutoHeight()
-		.Padding(10.f, 0.f, 10.f, 10.f)
-		[
-			SNew(SHorizontalBox)
+			// show buffers
 			+ SHorizontalBox::Slot()
-			.FillWidth(0.5f)
-			.Padding(0.f, 0.f, 10.f, 0.f)
+			.AutoWidth()
+			.VAlign(VAlign_Center)
+			.Padding(4.f)
 			[
-				SNew(SBorder)
+				SNew(SButton)
+				.ButtonStyle(FTimedDataMonitorEditorStyle::Get(), "ToggleButton")
+				.ToolTipText(LOCTEXT("ShowBuffers_ToolTip", "Open the buffer visualizer"))
+				.ContentPadding(FMargin(4, 2))
 				.HAlign(HAlign_Center)
 				.VAlign(VAlign_Center)
+				.OnClicked(this, &STimedDataMonitorPanel::OnShowBuffersClicked)
 				[
-					SNew(STextBlock)
-					.Text(LOCTEXT("Placeholder", "Placehoder for GL"))
+					SNew(SImage)
+					.Image(FTimedDataMonitorEditorStyle::Get().GetBrush("Img.BufferVisualization"))
 				]
 			]
+			// Settings button
 			+ SHorizontalBox::Slot()
-			.FillWidth(0.5f)
-			.Padding(0.f)
+			.Padding(4.f)
+			.AutoWidth()
+			.VAlign(VAlign_Center)
 			[
-				SNew(SBorder)
+				SNew(SButton)
+				.ButtonStyle(FTimedDataMonitorEditorStyle::Get(), "ToggleButton")
+				.ToolTipText(LOCTEXT("ShowUserSettings_Tip", "Show the general user settings"))
+				.ContentPadding(FMargin(4, 4))
 				.HAlign(HAlign_Center)
 				.VAlign(VAlign_Center)
+				.OnClicked(this, &STimedDataMonitorPanel::OnGeneralUserSettingsClicked)
 				[
-					SNew(STimecodeProvider)
-					.DisplayFrameRate(true)
+					SNew(STextBlock)
+					.Font(FEditorStyle::Get().GetFontStyle("FontAwesome.14"))
+					.Text(FEditorFontGlyphs::Cogs)
+					.ColorAndOpacity(FLinearColor::White)
 				]
 			]
-		]
-		// sources list
-		+ SVerticalBox::Slot()
-		.FillHeight(1.f)
-		.Padding(10.f, 0.f, 10.f, 10.f)
-		[
-			SNew(SScrollBox)
-			.Orientation(Orient_Vertical)
-			+ SScrollBox::Slot()
+			// Spacer
+			+ SHorizontalBox::Slot()
 			[
-				SAssignNew(TimedDataSourceList, STimedDataInputListView)
+				SNew(SSpacer)
 			]
-		]
-		// sources list
-		+ SVerticalBox::Slot()
-		.AutoHeight()
-		.Padding(10.f, 0.f, 10.f, 10.f)
+			// Status
+			+ SHorizontalBox::Slot()
+			.AutoWidth()
+			.VAlign(VAlign_Center)
+			.Padding(4.f)
+			[
+				SNew(STextBlock)
+				.Font(FEditorStyle::Get().GetFontStyle("FontAwesome.18"))
+				.Text(FEditorFontGlyphs::Circle)
+				.ColorAndOpacity(this, &STimedDataMonitorPanel::GetEvaluationStateColorAndOpacity)
+			]
+			// Status test
+			+ SHorizontalBox::Slot()
+			.AutoWidth()
+			.VAlign(VAlign_Center)
+			.Padding(4.f)
+			[
+				SNew(STextBlock)
+				.MinDesiredWidth(100)
+				.TextStyle(FTimedDataMonitorEditorStyle::Get(), "TextBlock.Large")
+				.Text(this, &STimedDataMonitorPanel::GetEvaluationStateText)
+				.ColorAndOpacity(FLinearColor::White)
+			]
+		];
+
+	ChildSlot
+	[
+		SNew(SOverlay)
+		+ SOverlay::Slot()
 		[
-			SNew(SBorder)
-			.Padding(FMargin(5.f, 5.f, 5.f, 5.f))
-			.Visibility(this, &STimedDataMonitorPanel::ShowMessageLog)
+			SNew(SVerticalBox)
+			// toolbar
+			+ SVerticalBox::Slot()
+			.AutoHeight()
+			.Padding(10.f, 10.f, 10.f, 10.f)
+			[
+				Toolbar
+			]
+			// timing element
+			+ SVerticalBox::Slot()
+			.AutoHeight()
+			.Padding(10.f, 0.f, 10.f, 10.f)
 			[
 				SNew(SHorizontalBox)
 				+ SHorizontalBox::Slot()
-				.FillWidth(1.0f)
+				.FillWidth(0.5f)
+				.Padding(0.f, 0.f, 10.f, 0.f)
+				[
+					SNew(SBorder)
+					.HAlign(HAlign_Center)
+					.VAlign(VAlign_Center)
+					.Padding(FMargin(4.f, 4.f, 4.f, 4.f))
+					[
+						SAssignNew(TimedDataGenlockWidget, STimedDataGenlock, SharedThis<STimedDataMonitorPanel>(this))
+					]
+				]
+				+ SHorizontalBox::Slot()
+				.FillWidth(0.5f)
+				.Padding(0.f)
 				[
-					MessageLogListingWidget
+					SNew(SBorder)
+					.HAlign(HAlign_Center)
+					.VAlign(VAlign_Center)
+					[
+						SAssignNew(TimedDataTimecodeWidget, STimedDataTimecodeProvider, SharedThis<STimedDataMonitorPanel>(this))
+					]
 				]
 			]
+			// sources list
+			+ SVerticalBox::Slot()
+			.FillHeight(1.f)
+			.Padding(10.f, 0.f, 10.f, 10.f)
+			[
+				SNew(SScrollBox)
+				.Orientation(Orient_Vertical)
+				+ SScrollBox::Slot()
+				[
+					SAssignNew(TimedDataSourceList, STimedDataInputListView, SharedThis<STimedDataMonitorPanel>(this))
+					.OnContextMenuOpening(this, &STimedDataMonitorPanel::OnDataListConstructContextMenu)
+				]
+			]
+			// message log
+			+ SVerticalBox::Slot()
+			.AutoHeight()
+			.Padding(10.f, 0.f, 10.f, 10.f)
+			[
+				SNew(SBorder)
+				.Padding(FMargin(5.f, 5.f, 5.f, 5.f))
+				.Visibility(this, &STimedDataMonitorPanel::ShowMessageLog)
+				[
+					SNew(SHorizontalBox)
+					+ SHorizontalBox::Slot()
+					.FillWidth(1.0f)
+					[
+						MessageLogListingWidget
+					]
+				]
+			]
+			// show performance warning
+			+SVerticalBox::Slot()
+			.AutoHeight()
+			.Padding(10.f, 0.f, 10.f, 10.f)
+			[
+				PerformanceWidget.ToSharedRef()
+			]
 		]
-		+SVerticalBox::Slot()
-		.AutoHeight()
-		.Padding(10.f, 0.f, 10.f, 10.f)
+		+ SOverlay::Slot()
 		[
 			SNew(SBorder)
 			.BorderImage(FEditorStyle::GetBrush("SettingsEditor.CheckoutWarningBorder"))
-			.BorderBackgroundColor(FColor(166, 137, 0))
+			.BorderBackgroundColor(FLinearColor(0.2f, 0.2f, 0.2f, 0.5f))
 			.Padding(FMargin(5.f, 5.f, 5.f, 5.f))
-			.Visibility(this, &STimedDataMonitorPanel::ShowEditorPerformanceThrottlingWarning)
+			.Visibility(this, &STimedDataMonitorPanel::GetThrobberVisibility)
 			[
-				SNew(SHorizontalBox)
-				+SHorizontalBox::Slot()
-				.FillWidth(1.0f)
-				.VAlign(VAlign_Center)
-				[
-					SNew(STextBlock)
-					.Text(PerformanceWarningText)
-					.Font(FEditorStyle::GetFontStyle("PropertyWindow.NormalFont"))
-					.ShadowColorAndOpacity(FLinearColor::Black.CopyWithNewOpacity(0.3f))
-					.ShadowOffset(FVector2D::UnitVector)
-				]
-				+SHorizontalBox::Slot()
-				.Padding(FMargin(0.f, 0.f, 10.f, 0.f))
-				.AutoWidth()
+				SNew(SBox)
+				.HAlign(HAlign_Center)
 				.VAlign(VAlign_Center)
+				.Clipping(EWidgetClipping::ClipToBounds)
 				[
-					SNew(SButton)
-					.OnClicked(this, &STimedDataMonitorPanel::DisableEditorPerformanceThrottling)
-					.Text(LOCTEXT("PerformanceWarningDisable", "Disable"))
+					SNew(SVerticalBox)
+					+ SVerticalBox::Slot()
+					.AutoHeight()
+					.HAlign(HAlign_Center)
+					[
+						SNew(SCircularThrobber)
+					]
+					+ SVerticalBox::Slot()
+					.AutoHeight()
+					.HAlign(HAlign_Center)
+					[
+						SNew(SButton)
+						.ContentPadding(FMargin(4, 2))
+						.OnClicked(this, &STimedDataMonitorPanel::OnCancelCalibration)
+						[
+							SNew(STextBlock)
+							.Text(LOCTEXT("CancelLabel", "Cancel"))
+						]
+					]
 				]
 			]
 		]
@@ -327,23 +460,50 @@ void STimedDataMonitorPanel::Construct(const FArguments& InArgs)
 END_SLATE_FUNCTION_BUILD_OPTIMIZATION
 
 
+void STimedDataMonitorPanel::Tick(const FGeometry& AllottedGeometry, const double InCurrentTime, const float InDeltaTime)
+{
+	double RefreshTimer = GetDefault<UTimedDataMonitorEditorSettings>()->RefreshRate;
+	if (bRefreshRequested || (FApp::GetCurrentTime() - LastCachedValueUpdateTime > RefreshTimer))
+	{
+		bRefreshRequested = false;
+		LastCachedValueUpdateTime = FApp::GetCurrentTime();
+
+		{
+			UTimedDataMonitorSubsystem* TimedDataMonitorSubsystem = GEngine->GetEngineSubsystem<UTimedDataMonitorSubsystem>();
+			check(TimedDataMonitorSubsystem);
+			CachedGlobalEvaluationState = TimedDataMonitorSubsystem->GetEvaluationState();
+		}
+
+		if (TimedDataGenlockWidget)
+		{
+			TimedDataGenlockWidget->UpdateCachedValue();
+		}
+		if (TimedDataTimecodeWidget)
+		{
+			TimedDataTimecodeWidget->UpdateCachedValue();
+		}
+		if (TimedDataSourceList)
+		{
+			TimedDataSourceList->UpdateCachedValue();
+		}
+	}
+
+	Super::Tick(AllottedGeometry, InCurrentTime, InDeltaTime);
+}
+
+
 void STimedDataMonitorPanel::BuildCalibrationArray()
 {
 	CalibrationUIAction[(int32)ETimedDataMonitorEditorCalibrationType::CalibrateWithTimecode] = FUIAction(
 		FExecuteAction::CreateSP(this, &STimedDataMonitorPanel::CalibrateWithTimecode));
-	CalibrationUIAction[(int32)ETimedDataMonitorEditorCalibrationType::JamWithPlatformTime] = FUIAction(
-		FExecuteAction::CreateSP(this, &STimedDataMonitorPanel::Jam, false));
-	CalibrationUIAction[(int32)ETimedDataMonitorEditorCalibrationType::JamWithTimecode] = FUIAction(
-		FExecuteAction::CreateSP(this, &STimedDataMonitorPanel::Jam, true));
-	CalibrationSlateIcon[(int32)ETimedDataMonitorEditorCalibrationType::CalibrateWithTimecode] = FSlateIcon(FTimedDataMonitorStyle::Get().GetStyleSetName(), "Img.Timecode.Small");
-	CalibrationSlateIcon[(int32)ETimedDataMonitorEditorCalibrationType::JamWithPlatformTime] = FSlateIcon(FTimedDataMonitorStyle::Get().GetStyleSetName(), "Img.Timecode.Small");
-	CalibrationSlateIcon[(int32)ETimedDataMonitorEditorCalibrationType::JamWithTimecode] = FSlateIcon(FTimedDataMonitorStyle::Get().GetStyleSetName(), "Img.Timecode.Small");
+	CalibrationUIAction[(int32)ETimedDataMonitorEditorCalibrationType::TimeCorrection] = FUIAction(
+		FExecuteAction::CreateSP(this, &STimedDataMonitorPanel::ApplyTimeCorrectionAll));
+	CalibrationSlateIcon[(int32)ETimedDataMonitorEditorCalibrationType::CalibrateWithTimecode] = FSlateIcon(FTimedDataMonitorEditorStyle::Get().GetStyleSetName(), "Img.Calibration");
+	CalibrationSlateIcon[(int32)ETimedDataMonitorEditorCalibrationType::TimeCorrection] = FSlateIcon(FTimedDataMonitorEditorStyle::Get().GetStyleSetName(), "Img.TimeCorrection");
 	CalibrationName[(int32)ETimedDataMonitorEditorCalibrationType::CalibrateWithTimecode] = LOCTEXT("CalibrateButtonLabel", "Calibrate");
-	CalibrationName[(int32)ETimedDataMonitorEditorCalibrationType::JamWithPlatformTime] = LOCTEXT("JamWithTimeButtonLabel", "Jam (with time)");
-	CalibrationName[(int32)ETimedDataMonitorEditorCalibrationType::JamWithTimecode] = LOCTEXT("JamWithTimecodeButtonLabel", "Jam (with timecode)");
+	CalibrationName[(int32)ETimedDataMonitorEditorCalibrationType::TimeCorrection] = LOCTEXT("JamWithTimeButtonLabel", "Time Correction");
 	CalibrationTooltip[(int32)ETimedDataMonitorEditorCalibrationType::CalibrateWithTimecode] = LOCTEXT("CalibrateButtonTooltip", "Find range of frames that all inputs has. If it can be found, change the Timecode Provider delay, resize buffers or delay them to accomadate.");
-	CalibrationTooltip[(int32)ETimedDataMonitorEditorCalibrationType::JamWithPlatformTime] = LOCTEXT("JamWithTimeButtonTooltip", "Set the offset of all inputs to match with the current Platform Time.");
-	CalibrationTooltip[(int32)ETimedDataMonitorEditorCalibrationType::JamWithTimecode] = LOCTEXT("JamWithTimecodeButtonTooltip", "Set the offset of all inputs to match with the current Timecode.");
+	CalibrationTooltip[(int32)ETimedDataMonitorEditorCalibrationType::TimeCorrection] = LOCTEXT("JamWithTimeButtonTooltip", "Set the offset of all inputs to match the Timecode value of the Timecode Provider or the Platform Time, depending on their evaluation type. It may resize buffers if they are too small.");
 }
 
 
@@ -379,6 +539,13 @@ FText STimedDataMonitorPanel::GetCalibrateButtonTooltip() const
 }
 
 
+const FSlateBrush* STimedDataMonitorPanel::GetCalibrateButtonImage() const
+{
+	ETimedDataMonitorEditorCalibrationType CalibrationType = GetDefault<UTimedDataMonitorEditorSettings>()->LastCalibrationType;
+	return CalibrationSlateIcon[(int32)CalibrationType].GetIcon();
+}
+
+
 FText STimedDataMonitorPanel::GetCalibrateButtonText() const
 {
 	ETimedDataMonitorEditorCalibrationType CalibrationType = GetDefault<UTimedDataMonitorEditorSettings>()->LastCalibrationType;
@@ -387,15 +554,203 @@ FText STimedDataMonitorPanel::GetCalibrateButtonText() const
 
 
 FReply STimedDataMonitorPanel::OnResetErrorsClicked()
+{
+	if (IsResetBufferStatChecked())
+	{
+		OnResetBufferStatClicked();
+	}
+	if (IsClearMessageChecked())
+	{
+		OnClearMessageClicked();
+	}
+	if (IsResetEvaluationChecked())
+	{
+		OnResetAllEvaluationTimeClicked();
+	}
+
+	return FReply::Handled();
+}
+
+
+TSharedRef<SWidget> STimedDataMonitorPanel::OnResetBuildMenu()
+{
+	FMenuBuilder MenuBuilder(true, nullptr);
+	//Reset the buffer statistics, message and evaluation time
+	MenuBuilder.AddMenuEntry(
+			LOCTEXT("ResetBufferStatisticsLabel", "Reset buffer statistics"),
+			LOCTEXT("ResetBufferStatisticsTooltip", "Reset the buffer statistics."),
+			FSlateIcon(),
+			FUIAction(
+				FExecuteAction::CreateSP(this, &STimedDataMonitorPanel::OnResetBufferStatClicked),
+				FCanExecuteAction(),
+				FIsActionChecked::CreateSP(this, &STimedDataMonitorPanel::IsResetBufferStatChecked))
+			);
+	MenuBuilder.AddMenuEntry(
+			LOCTEXT("ResetMessageLabel", "Clear messages"),
+			LOCTEXT("ResetMessageTooltip", "Clear the message list."),
+			FSlateIcon(),
+			FUIAction(
+				FExecuteAction::CreateSP(this, &STimedDataMonitorPanel::OnClearMessageClicked),
+				FCanExecuteAction(),
+				FIsActionChecked::CreateSP(this, &STimedDataMonitorPanel::IsClearMessageChecked))
+			);
+
+	MenuBuilder.AddMenuEntry(
+			LOCTEXT("ResetEvaluationTimeLabel", "Reset Evaluation Time"),
+			LOCTEXT("ResetEvaluationTimeTooltip", "Reset all the evaluation time for all input."),
+			FSlateIcon(),
+			FUIAction(
+				FExecuteAction::CreateSP(this, &STimedDataMonitorPanel::OnResetAllEvaluationTimeClicked),
+				FCanExecuteAction(),
+				FIsActionChecked::CreateSP(this, &STimedDataMonitorPanel::IsResetEvaluationChecked))
+			);
+
+	return MenuBuilder.MakeWidget();
+}
+
+
+void STimedDataMonitorPanel::OnResetBufferStatClicked()
 {
 	UTimedDataMonitorSubsystem* TimedDataMonitorSubsystem = GEngine->GetEngineSubsystem<UTimedDataMonitorSubsystem>();
 	check(TimedDataMonitorSubsystem);
 	TimedDataMonitorSubsystem->ResetAllBufferStats();
+}
+
+
+bool STimedDataMonitorPanel::IsResetBufferStatChecked() const
+{
+	return GetDefault<UTimedDataMonitorEditorSettings>()->bResetBufferStatEnabled;
+}
+
+
+void STimedDataMonitorPanel::OnClearMessageClicked()
+{
 	MessageLogListing->ClearMessages();
+}
+
+
+bool STimedDataMonitorPanel::IsClearMessageChecked() const
+{
+	return GetDefault<UTimedDataMonitorEditorSettings>()->bClearMessageEnabled;
+}
+
+
+void STimedDataMonitorPanel::OnResetAllEvaluationTimeClicked()
+{
+	UTimedDataMonitorSubsystem* TimedDataMonitorSubsystem = GEngine->GetEngineSubsystem<UTimedDataMonitorSubsystem>();
+	check(TimedDataMonitorSubsystem);
+	TArray<FTimedDataMonitorInputIdentifier> AllInputs = TimedDataMonitorSubsystem->GetAllInputs();
+	for (const FTimedDataMonitorInputIdentifier& Input : AllInputs)
+	{
+		TimedDataMonitorSubsystem->SetInputEvaluationOffsetInSeconds(Input, 0.f);
+	}
+}
+
+
+bool STimedDataMonitorPanel::IsResetEvaluationChecked() const
+{
+	return GetDefault<UTimedDataMonitorEditorSettings>()->bResetEvaluationTimeEnabled;
+}
+
+
+FReply STimedDataMonitorPanel::OnShowBuffersClicked()
+{
+	FGlobalTabmanager::Get()->InvokeTab(FTabId(Utilities::NAME_TimedDataMonitorBuffers));
+	return FReply::Handled();
+}
+
+
+FReply STimedDataMonitorPanel::OnGeneralUserSettingsClicked()
+{
+	FModuleManager::LoadModuleChecked<ISettingsModule>("Settings").ShowViewer("Editor", "Plugins", "Timed Data Monitor");
 	return FReply::Handled();
 }
 
 
+FSlateColor STimedDataMonitorPanel::GetEvaluationStateColorAndOpacity() const
+{
+	return (CachedGlobalEvaluationState == ETimedDataMonitorEvaluationState::InsideRange) ? FLinearColor::Green : FLinearColor::Red;
+}
+
+
+FText STimedDataMonitorPanel::GetEvaluationStateText() const
+{
+	switch(CachedGlobalEvaluationState)
+	{
+	case ETimedDataMonitorEvaluationState::NoSample:
+		return LOCTEXT("EvalState_NoSample", "No Samples");
+	case ETimedDataMonitorEvaluationState::OutsideRange:
+		return LOCTEXT("EvalState_OutsideRange", "Outside Range");
+	case ETimedDataMonitorEvaluationState::InsideRange:
+		return LOCTEXT("EvalState_InsideRange", "Synchronized");
+	}
+
+	return LOCTEXT("EvalState_Disabled", "Disabled");
+}
+
+
+TSharedPtr<SWidget> STimedDataMonitorPanel::OnDataListConstructContextMenu()
+{
+	const bool bShouldCloseWindowAfterMenuSelection = true;
+	FMenuBuilder MenuBuilder(bShouldCloseWindowAfterMenuSelection, nullptr);
+
+	MenuBuilder.AddMenuEntry(
+		LOCTEXT("TimeCorrection", "Time Correction"),
+		FText(),
+		CalibrationSlateIcon[(int32)ETimedDataMonitorEditorCalibrationType::TimeCorrection],
+		FUIAction(
+			FExecuteAction::CreateSP(this, &STimedDataMonitorPanel::ApplyTimeCorrectionOnSelection),
+			FCanExecuteAction::CreateSP(this, &STimedDataMonitorPanel::IsSourceListSectionValid)
+		)
+	);
+
+	MenuBuilder.AddMenuEntry(
+		LOCTEXT("ResetEvaluationTime", "Reset Time Correction"),
+		FText(),
+		FSlateIcon(),
+		FUIAction(
+			FExecuteAction::CreateSP(this, &STimedDataMonitorPanel::ResetTimeCorrectionOnSelection),
+			FCanExecuteAction::CreateSP(this, &STimedDataMonitorPanel::IsSourceListSectionValid)
+		)
+	);
+
+	return MenuBuilder.MakeWidget();
+}
+
+
+bool STimedDataMonitorPanel::IsSourceListSectionValid() const
+{
+	return TimedDataSourceList && TimedDataSourceList->GetSelectedInputIdentifier().IsValid();
+}
+
+
+void STimedDataMonitorPanel::ApplyTimeCorrectionOnSelection()
+{
+	if (TimedDataSourceList)
+	{
+		const FTimedDataMonitorInputIdentifier SelectedInput = TimedDataSourceList->GetSelectedInputIdentifier();
+		if (SelectedInput.IsValid())
+		{
+			ApplyTimeCorrection(SelectedInput);
+		}
+	}
+}
+
+
+void STimedDataMonitorPanel::ResetTimeCorrectionOnSelection()
+{
+	if (TimedDataSourceList)
+	{
+		const FTimedDataMonitorInputIdentifier SelectedInput = TimedDataSourceList->GetSelectedInputIdentifier();
+		if (SelectedInput.IsValid())
+		{
+			UTimedDataMonitorSubsystem* TimedDataMonitorSubsystem = GEngine->GetEngineSubsystem<UTimedDataMonitorSubsystem>();
+			TimedDataMonitorSubsystem->SetInputEvaluationOffsetInSeconds(SelectedInput, 0.f);
+		}
+	}
+}
+
+
 EVisibility STimedDataMonitorPanel::ShowMessageLog() const
 {
 	return MessageLogListing && MessageLogListing->NumMessages(EMessageSeverity::Info) > 0 ? EVisibility::Visible : EVisibility::Collapsed;
@@ -419,123 +774,215 @@ FReply STimedDataMonitorPanel::DisableEditorPerformanceThrottling()
 }
 
 
-void STimedDataMonitorPanel::CalibrateWithTimecode()
+EVisibility STimedDataMonitorPanel::GetThrobberVisibility() const
 {
-	UTimedDataMonitorSubsystem* TimedDataMonitorSubsystem = GEngine->GetEngineSubsystem<UTimedDataMonitorSubsystem>();
-	check(TimedDataMonitorSubsystem);
+	return MonitorCalibration ? EVisibility::Visible : EVisibility::Collapsed;
+}
+
 
+void STimedDataMonitorPanel::CalibrateWithTimecode()
+{
 	MessageLogListing->ClearMessages();
+	
+	FTimedDataMonitorCalibrationParameters Params = GetDefault<UTimedDataMonitorEditorSettings>()->CalibrationSettings;
+	FTimedDataMonitorCalibration::FOnCalibrationCompletedSignature CompletedCallback = FTimedDataMonitorCalibration::FOnCalibrationCompletedSignature::CreateSP(this, &STimedDataMonitorPanel::CalibrateWithTimecodeCompleted);
 
-	FTimedDataMonitorCalibrationResult Result = TimedDataMonitorSubsystem->CalibrateWithTimecodeProvider();
-
-	//if (Result.ReturnCode != ETimedDataMonitorCalibrationReturnCode::Succeeded)
-	//{
-	//	switch (Result.ReturnCode)
-	//	{
-	//	case ETimedDataMonitorCalibrationReturnCode::Failed_NoTimecode:
-	//		FMessageLog(Utilities::NAME_LogName).Error(LOCTEXT("CalibrationFailed_NoTimecode", "The timecode provider doesn't have a proper timecode value."));
-	//		break;
-	//	case ETimedDataMonitorCalibrationReturnCode::Failed_UnresponsiveInput:
-	//		for (const FTimedDataMonitorChannelIdentifier& Identifier : Result.FailureChannelIdentifiers)
-	//		{
-	//			FMessageLog(Utilities::NAME_LogName).Error(FText::Format(LOCTEXT("CalibrationFailed_UnresponsiveInput", "The input '{0}' is unresponsive."), TimedDataMonitorSubsystem->GetChannelDisplayName(Identifier)));
-	//		}
-	//		break;
-	//	case ETimedDataMonitorCalibrationReturnCode::Failed_InvalidFrameRate:
-	//		for (const FTimedDataMonitorChannelIdentifier& Identifier : Result.FailureChannelIdentifiers)
-	//		{
-	//			FMessageLog(Utilities::NAME_LogName).Error(FText::Format(LOCTEXT("CalibrationFailed_InvalidFrameRate", "The input '{0}' has an invalid frame rate."), TimedDataMonitorSubsystem->GetChannelDisplayName(Identifier)));
-	//		}
-	//		break;
-	//	case ETimedDataMonitorCalibrationReturnCode::Failed_NoDataBuffered:
-	//		for (const FTimedDataMonitorChannelIdentifier& Identifier : Result.FailureChannelIdentifiers)
-	//		{
-	//			FMessageLog(Utilities::NAME_LogName).Error(FText::Format(LOCTEXT("CalibrationFailed_NoDataBuffered", "The input '{0}' has not data buffured."), TimedDataMonitorSubsystem->GetChannelDisplayName(Identifier)));
-	//		}
-	//		break;
-	//	default:
-	//		break;
-	//	}
-	//}
+	MonitorCalibration = MakeUnique<FTimedDataMonitorCalibration>();
+	MonitorCalibration->CalibrateWithTimecode(Params, CompletedCallback);
 
 	GetMutableDefault<UTimedDataMonitorEditorSettings>()->LastCalibrationType = ETimedDataMonitorEditorCalibrationType::CalibrateWithTimecode;
 	GetMutableDefault<UTimedDataMonitorEditorSettings>()->SaveConfig();
 }
 
+void STimedDataMonitorPanel::CalibrateWithTimecodeCompleted(FTimedDataMonitorCalibrationResult Result)
+{
+	UTimedDataMonitorSubsystem* TimedDataMonitorSubsystem = GEngine->GetEngineSubsystem<UTimedDataMonitorSubsystem>();
+	MonitorCalibration.Reset();
+
+	if (Result.ReturnCode != ETimedDataMonitorCalibrationReturnCode::Succeeded)
+	{
+		switch (Result.ReturnCode)
+		{
+		case ETimedDataMonitorCalibrationReturnCode::Succeeded:
+			MessageLogListing->AddMessage(FTokenizedMessage::Create(EMessageSeverity::Info
+				, LOCTEXT("CalibrateSucceeded", "Calibration succeeded.")));
+			break;
+		case ETimedDataMonitorCalibrationReturnCode::Failed_NoTimecode:
+			MessageLogListing->AddMessage(FTokenizedMessage::Create(EMessageSeverity::Error
+				, LOCTEXT("CalibrateFailed_NoTimecode", "The timecode provider doesn't have a proper timecode value.")));
+			break;
+		case ETimedDataMonitorCalibrationReturnCode::Failed_UnresponsiveInput:
+			for (const FTimedDataMonitorInputIdentifier& Identifier : Result.FailureInputIdentifiers)
+			{
+				MessageLogListing->AddMessage(FTokenizedMessage::Create(EMessageSeverity::Error
+					, FText::Format(LOCTEXT("CalibrateFailed_UnresponsiveInput", "The input '{0}' is unresponsive."), TimedDataMonitorSubsystem->GetInputDisplayName(Identifier))));
+			}
+			break;
+		case ETimedDataMonitorCalibrationReturnCode::Failed_InvalidEvaluationType:
+			for (const FTimedDataMonitorInputIdentifier& Identifier : Result.FailureInputIdentifiers)
+			{
+				MessageLogListing->AddMessage(FTokenizedMessage::Create(EMessageSeverity::Error
+					, FText::Format(LOCTEXT("CalibrationFailed_InvalidEvaluationType", "The input '{0}' is not evaluated in timecode."), TimedDataMonitorSubsystem->GetInputDisplayName(Identifier))));
+			}
+			break;
+		case ETimedDataMonitorCalibrationReturnCode::Failed_InvalidFrameRate:
+			for (const FTimedDataMonitorInputIdentifier& Identifier : Result.FailureInputIdentifiers)
+			{
+				MessageLogListing->AddMessage(FTokenizedMessage::Create(EMessageSeverity::Error
+					, FText::Format(LOCTEXT("CalibrationFailed_InvalidFrameRate", "The input '{0}' has an invalid frame rate."), TimedDataMonitorSubsystem->GetInputDisplayName(Identifier))));
+			}
+			break;
+		case ETimedDataMonitorCalibrationReturnCode::Failed_NoDataBuffered:
+			for (const FTimedDataMonitorInputIdentifier& Identifier : Result.FailureInputIdentifiers)
+			{
+				MessageLogListing->AddMessage(FTokenizedMessage::Create(EMessageSeverity::Error
+					, FText::Format(LOCTEXT("CalibrationFailed_NoDataBuffered", "The input '{0}' has at least one channel that does not data buffured."), TimedDataMonitorSubsystem->GetInputDisplayName(Identifier))));
+			}
+			break;
+		case ETimedDataMonitorCalibrationReturnCode::Failed_Reset:
+			MessageLogListing->AddMessage(FTokenizedMessage::Create(EMessageSeverity::Error
+				, LOCTEXT("CalibrationFailed_Reset", "The process was manually interrupted.")));
+			break;
+		case ETimedDataMonitorCalibrationReturnCode::Retry_NotEnoughData:
+			for (const FTimedDataMonitorInputIdentifier& Identifier : Result.FailureInputIdentifiers)
+			{
+				MessageLogListing->AddMessage(FTokenizedMessage::Create(EMessageSeverity::Error
+					, FText::Format(LOCTEXT("CalibrationFailed_IncreaseBufferSize", "No interval could be found. The buffer size that is required for input '{0}' should be enough. It doesn't contain enough data."), TimedDataMonitorSubsystem->GetInputDisplayName(Identifier))));
+			}
+			break;
+		case ETimedDataMonitorCalibrationReturnCode::Retry_IncreaseBufferSize:
+			for (const FTimedDataMonitorInputIdentifier& Identifier : Result.FailureInputIdentifiers)
+			{
+				MessageLogListing->AddMessage(FTokenizedMessage::Create(EMessageSeverity::Warning
+					, FText::Format(LOCTEXT("CalibrationRetry_IncreaseBufferSize", "No interval could be found. The buffer size of input '{0}' has been increased. You may retry when the buffer will be full."), TimedDataMonitorSubsystem->GetInputDisplayName(Identifier))));
+			}
+			break;
+		default:
+			break;
+		}
+	}
+}
+
 
-void STimedDataMonitorPanel::Jam(bool bWithTimecode)
+void STimedDataMonitorPanel::ApplyTimeCorrectionAll()
 {
 	UTimedDataMonitorSubsystem* TimedDataMonitorSubsystem = GEngine->GetEngineSubsystem<UTimedDataMonitorSubsystem>();
 	check(TimedDataMonitorSubsystem);
 
 	MessageLogListing->ClearMessages();
 
-	ETimedDataInputEvaluationType JamType = bWithTimecode ? ETimedDataInputEvaluationType::Timecode : ETimedDataInputEvaluationType::PlatformTime;
-	FTimedDataMonitorJamResult Result = TimedDataMonitorSubsystem->JamInputs(JamType);
+	TArray<FTimedDataMonitorInputIdentifier> AllInputs = TimedDataMonitorSubsystem->GetAllInputs();
+	bool bAllSucceeded = true;
+	for (const FTimedDataMonitorInputIdentifier& InputIndentifier : AllInputs)
+	{
+		if (TimedDataMonitorSubsystem->GetInputEnabled(InputIndentifier) != ETimedDataMonitorInputEnabled::Disabled)
+		{
+			ETimedDataMonitorTimeCorrectionReturnCode Result = ApplyTimeCorrection(InputIndentifier);
+			bAllSucceeded = bAllSucceeded && (Result == ETimedDataMonitorTimeCorrectionReturnCode::Succeeded);
+		}
+	}
 
-	switch (Result.ReturnCode)
+	if (bAllSucceeded)
 	{
-	case ETimedDataMonitorJamReturnCode::Succeeded:
 		MessageLogListing->AddMessage(FTokenizedMessage::Create(EMessageSeverity::Info
-			, LOCTEXT("JamSucceeded_NoTimecode", "Jam succeeded.")));
+			, LOCTEXT("JamSucceeded", "Time Correction succeeded.")));
+	}
+
+	GetMutableDefault<UTimedDataMonitorEditorSettings>()->LastCalibrationType = ETimedDataMonitorEditorCalibrationType::CalibrateWithTimecode;
+	GetMutableDefault<UTimedDataMonitorEditorSettings>()->SaveConfig();
+}
+
+
+ETimedDataMonitorTimeCorrectionReturnCode STimedDataMonitorPanel::ApplyTimeCorrection(const FTimedDataMonitorInputIdentifier& InputIdentifier)
+{
+	UTimedDataMonitorSubsystem* TimedDataMonitorSubsystem = GEngine->GetEngineSubsystem<UTimedDataMonitorSubsystem>();
+	check(TimedDataMonitorSubsystem);
+
+	const FTimedDataMonitorTimeCorrectionResult Result = FTimedDataMonitorCalibration::ApplyTimeCorrection(InputIdentifier, GetDefault<UTimedDataMonitorEditorSettings>()->TimeCorrectionSettings);
+	switch (Result.ReturnCode)
+	{
+	case ETimedDataMonitorTimeCorrectionReturnCode::Failed_InvalidInput:
+		MessageLogListing->AddMessage(FTokenizedMessage::Create(EMessageSeverity::Error
+			, LOCTEXT("JamFailed_InvalidInput", "The input identifier was invalid.")));
 		break;
-	case ETimedDataMonitorJamReturnCode::Failed_NoTimecode:
+	case ETimedDataMonitorTimeCorrectionReturnCode::Failed_NoTimecode:
 		MessageLogListing->AddMessage(FTokenizedMessage::Create(EMessageSeverity::Error
 			, LOCTEXT("JamFailed_NoTimecode", "The timecode provider doesn't have a proper timecode value.")));
 		break;
-	case ETimedDataMonitorJamReturnCode::Failed_UnresponsiveInput:
-		for (const FTimedDataMonitorChannelIdentifier& Identifier : Result.FailureChannelIdentifiers)
+	case ETimedDataMonitorTimeCorrectionReturnCode::Failed_UnresponsiveInput:
+		for (const FTimedDataMonitorChannelIdentifier& ChannelIdentifier : Result.FailureChannelIdentifiers)
 		{
 			MessageLogListing->AddMessage(FTokenizedMessage::Create(EMessageSeverity::Error
-				, FText::Format(LOCTEXT("JamFailed_UnresponsiveInput", "The channel '{0}' is unresponsive."), TimedDataMonitorSubsystem->GetChannelDisplayName(Identifier))));
+				, FText::Format(LOCTEXT("JamFailed_UnresponsiveInput", "The channel '{0}' is unresponsive."), TimedDataMonitorSubsystem->GetChannelDisplayName(ChannelIdentifier))));
 		}
 		break;
-	case ETimedDataMonitorJamReturnCode::Failed_EvaluationTypeDoNotMatch:
-		for (const FTimedDataMonitorInputIdentifier& Identifier : Result.FailureInputIdentifiers)
+	case ETimedDataMonitorTimeCorrectionReturnCode::Failed_NoDataBuffered:
+		for (const FTimedDataMonitorChannelIdentifier& ChannelIdentifier : Result.FailureChannelIdentifiers)
 		{
 			MessageLogListing->AddMessage(FTokenizedMessage::Create(EMessageSeverity::Error
-				, FText::Format(LOCTEXT("JamFailed_InvalidFrameRate", "The input '{0}' evalution type doesn't match with the jam type."), TimedDataMonitorSubsystem->GetInputDisplayName(Identifier))));
+				, FText::Format(LOCTEXT("JamFailed_NoDataBuffered", "The channel '{0}' has not data buffered."), TimedDataMonitorSubsystem->GetChannelDisplayName(ChannelIdentifier))));
 		}
 		break;
-	case ETimedDataMonitorJamReturnCode::Failed_NoDataBuffered:
-		for (const FTimedDataMonitorChannelIdentifier& Identifier : Result.FailureChannelIdentifiers)
+	case ETimedDataMonitorTimeCorrectionReturnCode::Failed_BufferCouldNotBeResize:
+		if (Result.FailureChannelIdentifiers.Num() > 0)
+		{
+			for (const FTimedDataMonitorChannelIdentifier& ChannelIdentifier : Result.FailureChannelIdentifiers)
+			{
+				MessageLogListing->AddMessage(FTokenizedMessage::Create(EMessageSeverity::Error
+					, FText::Format(LOCTEXT("JamFailed_BufferSizeHaveBeenMaxed_Channel", "The buffer size of channel '{0}' could not be increased."), TimedDataMonitorSubsystem->GetChannelDisplayName(ChannelIdentifier))));
+			}
+		}
+		else
 		{
 			MessageLogListing->AddMessage(FTokenizedMessage::Create(EMessageSeverity::Error
-				, FText::Format(LOCTEXT("JamFailed_NoDataBuffered", "The channel '{0}' has not data buffered."), TimedDataMonitorSubsystem->GetChannelDisplayName(Identifier))));
+				, FText::Format(LOCTEXT("JamFailed_BufferSizeHaveBeenMaxed_Input", "The buffer size of input '{0}' could not be increased."), TimedDataMonitorSubsystem->GetInputDisplayName(InputIdentifier))));
 		}
 		break;
-	case ETimedDataMonitorJamReturnCode::Failed_BufferSizeHaveBeenMaxed:
-		for (const FTimedDataMonitorInputIdentifier& Identifier : Result.FailureInputIdentifiers)
+	case ETimedDataMonitorTimeCorrectionReturnCode::Retry_NotEnoughData:
+		if (Result.FailureChannelIdentifiers.Num() > 0)
 		{
-			MessageLogListing->AddMessage(FTokenizedMessage::Create(EMessageSeverity::Error
-				, FText::Format(LOCTEXT("Failed_BufferSizeHaveBeenMaxed_Input", "The buffer size of input '{0}' could not be increased further."), TimedDataMonitorSubsystem->GetInputDisplayName(Identifier))));
+			for (const FTimedDataMonitorChannelIdentifier& ChannelIdentifier : Result.FailureChannelIdentifiers)
+			{
+				MessageLogListing->AddMessage(FTokenizedMessage::Create(EMessageSeverity::Error
+					, FText::Format(LOCTEXT("JamRetry_IncreaseBufferSize_Channel", "No interval could be found. The buffer  that is required for channel '{0}' could not be increased."), TimedDataMonitorSubsystem->GetChannelDisplayName(ChannelIdentifier))));
+			}
 		}
-		for (const FTimedDataMonitorChannelIdentifier& Identifier : Result.FailureChannelIdentifiers)
+		else
 		{
 			MessageLogListing->AddMessage(FTokenizedMessage::Create(EMessageSeverity::Error
-				, FText::Format(LOCTEXT("Failed_BufferSizeHaveBeenMaxed_Channel", "The buffer size of channel '{0}' needed to be increased further."), TimedDataMonitorSubsystem->GetChannelDisplayName(Identifier))));
+				, FText::Format(LOCTEXT("JamRetry_IncreaseBufferSize_Input", "No interval could be found. The buffer size that is required for input '{0}' should be enough. It doesn't contain enough data."), TimedDataMonitorSubsystem->GetInputDisplayName(InputIdentifier))));
 		}
-		break;
-	case ETimedDataMonitorJamReturnCode::Retry_BufferSizeHasBeenIncreased:
-		for (const FTimedDataMonitorInputIdentifier& Identifier : Result.FailureInputIdentifiers)
+		break;				
+	case ETimedDataMonitorTimeCorrectionReturnCode::Retry_IncreaseBufferSize:
+		if (Result.FailureChannelIdentifiers.Num() > 0)
 		{
-			MessageLogListing->AddMessage(FTokenizedMessage::Create(EMessageSeverity::Warning
-				, FText::Format(LOCTEXT("JamRetry_BufferSizeHasBeenIncrease_Input", "The buffer size of input '{0}' needed to be increased. Retry."), TimedDataMonitorSubsystem->GetInputDisplayName(Identifier))));
+			for (const FTimedDataMonitorChannelIdentifier& Identifier : Result.FailureChannelIdentifiers)
+			{
+				MessageLogListing->AddMessage(FTokenizedMessage::Create(EMessageSeverity::Warning
+					, FText::Format(LOCTEXT("JamRetry_BufferSizeHasBeenIncrease_Channel", "The buffer size of channel '{0}' needed to be increased. Retry."), TimedDataMonitorSubsystem->GetChannelDisplayName(Identifier))));
+			}
 		}
-		for (const FTimedDataMonitorChannelIdentifier& Identifier : Result.FailureChannelIdentifiers)
+		else
 		{
-			MessageLogListing->AddMessage(FTokenizedMessage::Create(EMessageSeverity::Warning
-				, FText::Format(LOCTEXT("JamRetry_BufferSizeHasBeenIncrease_Channel", "The buffer size of channel '{0}' needed to be increased. Retry."), TimedDataMonitorSubsystem->GetChannelDisplayName(Identifier))));
+			MessageLogListing->AddMessage(FTokenizedMessage::Create(EMessageSeverity::Error
+				, FText::Format(LOCTEXT("JamRetry_BufferSizeHasBeenIncrease_Input", "The buffer size of input '{0}' needed to be increased. Retry."), TimedDataMonitorSubsystem->GetInputDisplayName(InputIdentifier))));
 		}
 		break;
 	default:
 		break;
 	}
 
-	GetMutableDefault<UTimedDataMonitorEditorSettings>()->LastCalibrationType = bWithTimecode
-		? ETimedDataMonitorEditorCalibrationType::JamWithTimecode
-		: ETimedDataMonitorEditorCalibrationType::JamWithPlatformTime;
-	GetMutableDefault<UTimedDataMonitorEditorSettings>()->SaveConfig();
+	return Result.ReturnCode;
+}
+
+
+FReply STimedDataMonitorPanel::OnCancelCalibration()
+{
+	if (MonitorCalibration)
+	{
+		MonitorCalibration.Reset();
+	}
+	return FReply::Handled();
 }
 
 #undef LOCTEXT_NAMESPACE
diff --git a/Engine/Plugins/VirtualProduction/TimedDataMonitor/Source/TimedDataMonitorEditor/Private/STimedDataMonitorPanel.h b/Engine/Plugins/VirtualProduction/TimedDataMonitor/Source/TimedDataMonitorEditor/Private/STimedDataMonitorPanel.h
--- a/Engine/Plugins/VirtualProduction/TimedDataMonitor/Source/TimedDataMonitorEditor/Private/STimedDataMonitorPanel.h
+++ b/Engine/Plugins/VirtualProduction/TimedDataMonitor/Source/TimedDataMonitorEditor/Private/STimedDataMonitorPanel.h
@@ -15,12 +15,17 @@
 #include "TimedDataMonitorEditorSettings.h"
 
 
-
+struct FSlateBrush;
+struct FTimedDataMonitorCalibrationResult;
+class FTimedDataMonitorCalibration;
 class FWorkspaceItem;
 class IMessageLogListing;
+class STimedDataGenlock;
 class STimedDataInputListView;
+class STimedDataTimecodeProvider;
 class SWidget;
 
+enum class ETimedDataMonitorEvaluationState : uint8;
 
 class STimedDataMonitorPanel : public SCompoundWidget
 {
@@ -39,22 +44,49 @@ public:
 
 	void Construct(const FArguments& InArgs);
 
+	virtual void Tick(const FGeometry& AllottedGeometry, const double InCurrentTime, const float InDeltaTime) override;
+
+	void RequestRefresh() { bRefreshRequested = true; }
+
 private:
 	FReply OnCalibrateClicked();
 	TSharedRef<SWidget> OnCalibrateBuildMenu();
 	FText GetCalibrateButtonTooltip() const;
+	const FSlateBrush* GetCalibrateButtonImage() const;
 	FText GetCalibrateButtonText() const;
 	FReply OnResetErrorsClicked();
+	TSharedRef<SWidget> OnResetBuildMenu();
+	void OnResetBufferStatClicked();
+	bool IsResetBufferStatChecked() const;
+	void OnClearMessageClicked();
+	bool IsClearMessageChecked() const;
+	void OnResetAllEvaluationTimeClicked();
+	bool IsResetEvaluationChecked() const;
+	FReply OnShowBuffersClicked();
+	FReply OnGeneralUserSettingsClicked();
+	FSlateColor GetEvaluationStateColorAndOpacity() const;
+	FText GetEvaluationStateText() const;
+	TSharedPtr<SWidget> OnDataListConstructContextMenu();
+	bool IsSourceListSectionValid() const;
+	void ApplyTimeCorrectionOnSelection();
+	void ResetTimeCorrectionOnSelection();
 
 	EVisibility ShowMessageLog() const;
 	EVisibility ShowEditorPerformanceThrottlingWarning() const;
 	FReply DisableEditorPerformanceThrottling();
 
+	EVisibility GetThrobberVisibility() const;
+
 	void BuildCalibrationArray();
 	void CalibrateWithTimecode();
-	void Jam(bool bWithTimecode);
+	void CalibrateWithTimecodeCompleted(FTimedDataMonitorCalibrationResult);
+	void ApplyTimeCorrectionAll();
+	ETimedDataMonitorTimeCorrectionReturnCode ApplyTimeCorrection(const FTimedDataMonitorInputIdentifier& InputIndentifier);
+	FReply OnCancelCalibration();
 
 private:
+	TSharedPtr<STimedDataGenlock> TimedDataGenlockWidget;
+	TSharedPtr<STimedDataTimecodeProvider> TimedDataTimecodeWidget;
 	TSharedPtr<STimedDataInputListView> TimedDataSourceList;
 	TSharedPtr<IMessageLogListing> MessageLogListing;
 
@@ -64,6 +96,12 @@ private:
 	FText CalibrationName[CalibrationArrayCount];
 	FText CalibrationTooltip[CalibrationArrayCount];
 
-	static FDelegateHandle LevelEditorTabManagerChangedHandle;
+	ETimedDataMonitorEvaluationState CachedGlobalEvaluationState;
+
+	TUniquePtr<FTimedDataMonitorCalibration> MonitorCalibration;
 
+	bool bRefreshRequested = true;
+	double LastCachedValueUpdateTime = 0.0;
+
+	static FDelegateHandle LevelEditorTabManagerChangedHandle;
 };
\ No newline at end of file
diff --git a/Engine/Plugins/VirtualProduction/TimedDataMonitor/Source/TimedDataMonitorEditor/Private/STimedDataNumericEntryBox.h b/Engine/Plugins/VirtualProduction/TimedDataMonitor/Source/TimedDataMonitorEditor/Private/STimedDataNumericEntryBox.h
new file mode 100644
--- /dev/null
+++ b/Engine/Plugins/VirtualProduction/TimedDataMonitor/Source/TimedDataMonitorEditor/Private/STimedDataNumericEntryBox.h
@@ -0,0 +1,206 @@
+// Copyright Epic Games, Inc. All Rights Reserved.
+
+#pragma once
+
+#include "CoreMinimal.h"
+
+#include "SlateFwd.h"
+#include "Widgets/DeclarativeSyntaxSupport.h"
+#include "Widgets/SCompoundWidget.h"
+#include "SlateOptMacros.h"
+
+#include "EditorFontGlyphs.h"
+#include "EditorStyleSet.h"
+#include "TimedDataMonitorEditorStyle.h"
+
+#include "Widgets/Images/SImage.h"
+#include "Widgets/Input/SCheckBox.h"
+#include "Widgets/Input/SButton.h"
+#include "Widgets/Input/SNumericEntryBox.h"
+#include "Widgets/Layout/SBorder.h"
+#include "Widgets/Text/STextBlock.h"
+
+
+
+template<class NumericType>
+class STimedDataNumericEntryBox : public SCompoundWidget
+{
+private:
+	using Super = SCompoundWidget;
+	using FOnValueCommitted = typename SNumericEntryBox<NumericType>::FOnValueCommitted;
+	
+public:
+	SLATE_BEGIN_ARGS(STimedDataNumericEntryBox)
+			: _Value(0)
+			, _MinValue(0)
+			, _CanEdit(true)
+			, _Amount(0)
+			, _ShowAmount(false)
+		{ }
+		SLATE_ATTRIBUTE(NumericType, Value)
+		SLATE_ARGUMENT(NumericType, MinValue)
+		SLATE_ARGUMENT(FText, EditLabel)
+		SLATE_ARGUMENT(bool, CanEdit);
+		SLATE_ATTRIBUTE(NumericType, Amount);
+		SLATE_ARGUMENT(bool, ShowAmount);
+		SLATE_STYLE_ARGUMENT(FTextBlockStyle, TextStyle)
+		SLATE_EVENT(FOnValueCommitted, OnValueCommitted)
+	SLATE_END_ARGS()
+	
+public:
+	void Construct(const FArguments& InArgs)
+	{
+		Value = InArgs._Value;
+		MinValue = InArgs._MinValue;
+		EditLabel = InArgs._EditLabel;
+		Amount = InArgs._Amount;
+		bShowAmount = InArgs._ShowAmount;
+		OnValueCommitted = InArgs._OnValueCommitted;
+
+		TSharedRef<STextBlock> TextBlock = SNew(STextBlock)
+			.TextStyle(InArgs._TextStyle)
+			.Text(this, &STimedDataNumericEntryBox::GetValueText);
+
+		if (InArgs._CanEdit)
+		{
+			ChildSlot
+			[
+				SNew(SHorizontalBox)
+				+ SHorizontalBox::Slot()
+				.HAlign(HAlign_Right)
+				.VAlign(VAlign_Center)
+				.FillWidth(1.f)
+				.Padding(4.0f, 0.0f, 0.0f, 0.0f)
+				[
+					TextBlock
+				]
+				+ SHorizontalBox::Slot()
+				.HAlign(HAlign_Right)
+				.VAlign(VAlign_Center)
+				.AutoWidth()
+				.Padding(4.0f, 0.0f, 0.0f, 0.0f)
+				[
+					SAssignNew(ComboButton, SComboButton)
+					.ComboButtonStyle(FTimedDataMonitorEditorStyle::Get(), "ToggleComboButton")
+					.HAlign(HAlign_Center)
+					.HasDownArrow(false)
+					.OnGetMenuContent(this, &STimedDataNumericEntryBox::OnCreateEditMenu)
+					.ButtonContent()
+					[
+						SNew(STextBlock)
+						.Font(FEditorStyle::Get().GetFontStyle("FontAwesome.8"))
+						.Text(FEditorFontGlyphs::Pencil_Square)
+						.ColorAndOpacity(FLinearColor::White)
+					]
+				]
+			];
+		}
+		else
+		{
+			ChildSlot
+			[
+				TextBlock
+			];
+		}
+	}
+	
+private:
+	FText GetValueText() const 
+	{
+		if (bShowAmount)
+		{
+			return FText::Format(NSLOCTEXT("TimedDataNumericEntryBox", "TimedNumericDataState", "{0}/{1}"), Amount.Get(), Value.Get());
+		}
+		else
+		{
+			return FText::AsNumber(Value.Get()); 
+		}
+	}
+
+	TSharedRef<SWidget> OnCreateEditMenu()
+	{
+		bCloseRequested = false;
+
+		TSharedPtr<SWidget> WidgetToFocus;
+		TSharedRef<SWidget> Result = SNew(SBorder)
+			.VAlign(VAlign_Center)
+			.Padding(FMargin(5.f, 5.f, 5.f, 5.f))
+			[
+				SNew(SHorizontalBox)
+				+ SHorizontalBox::Slot()
+				.HAlign(HAlign_Fill)
+				.VAlign(VAlign_Center)
+				.AutoWidth()
+				[
+					SNew(STextBlock)
+					.TextStyle(FTimedDataMonitorEditorStyle::Get(), "TextBlock.Regular")
+					.Text(EditLabel)
+				]
+				+ SHorizontalBox::Slot()
+				.HAlign(HAlign_Fill)
+				.VAlign(VAlign_Center)
+				.FillWidth(1.f)
+				.Padding(4.0f, 0.0f, 0.0f, 0.0f)
+				[
+					SAssignNew(WidgetToFocus, SNumericEntryBox<NumericType>)
+					.MinValue(TOptional<NumericType>(MinValue))
+					.MinDesiredValueWidth(50)
+					.Value(TOptional<NumericType>(Value.Get()))
+					.OnValueCommitted(this, &STimedDataNumericEntryBox::OnValueComittedCallback)
+				]
+				+ SHorizontalBox::Slot()
+				.HAlign(HAlign_Right)
+				.VAlign(VAlign_Center)
+				.AutoWidth()
+				.Padding(4.0f, 0.0f, 0.0f, 0.0f)
+				[
+					SNew(SCheckBox)
+					.Padding(4.f)
+					.Style(FEditorStyle::Get(), "ToggleButtonCheckbox")
+					.IsChecked_Lambda([]() {return ECheckBoxState::Unchecked; })
+					.OnCheckStateChanged(this, &STimedDataNumericEntryBox::CloseComboButton)
+					[
+						SNew(STextBlock)
+						.Font(FEditorStyle::Get().GetFontStyle("FontAwesome.11"))
+						.Text(FEditorFontGlyphs::Check)
+					]
+				]
+			];
+
+		ComboButton->SetMenuContentWidgetToFocus(WidgetToFocus);
+		return Result;
+	}
+
+	void CloseComboButton(ECheckBoxState InNewState)
+	{
+		if (ComboButton)
+		{
+			ComboButton->SetIsOpen(false);
+		}
+	}
+
+	void OnValueComittedCallback(NumericType NewValue, ETextCommit::Type CommitType)
+	{
+		if (!bCloseRequested)
+		{
+			bCloseRequested = true;
+
+			OnValueCommitted.ExecuteIfBound(NewValue, CommitType);
+			if (ComboButton)
+			{
+				ComboButton->SetIsOpen(false);
+			}
+		}
+	}
+
+private:
+	TAttribute<NumericType> Value;
+	TAttribute<NumericType> Amount;
+	NumericType MinValue;
+	FText EditLabel;
+	bool bShowAmount;
+	FOnValueCommitted OnValueCommitted;
+	TSharedPtr<SComboButton> ComboButton;
+	bool bCloseRequested = false;
+};
+
diff --git a/Engine/Plugins/VirtualProduction/TimedDataMonitor/Source/TimedDataMonitorEditor/Private/TimedDataMonitorEditorSettings.h b/Engine/Plugins/VirtualProduction/TimedDataMonitor/Source/TimedDataMonitorEditor/Private/TimedDataMonitorEditorSettings.h
--- a/Engine/Plugins/VirtualProduction/TimedDataMonitor/Source/TimedDataMonitorEditor/Private/TimedDataMonitorEditorSettings.h
+++ b/Engine/Plugins/VirtualProduction/TimedDataMonitor/Source/TimedDataMonitorEditor/Private/TimedDataMonitorEditorSettings.h
@@ -6,6 +6,8 @@
 #include "UObject/Object.h"
 #include "UObject/ObjectMacros.h"
 
+#include "TimedDataMonitorCalibration.h"
+
 #include "TimedDataMonitorEditorSettings.generated.h"
 
 
@@ -13,9 +15,8 @@ UENUM()
 enum class ETimedDataMonitorEditorCalibrationType
 {
 	CalibrateWithTimecode = 0,
-	JamWithTimecode = 1,
-	JamWithPlatformTime = 2,
-	Max = 3
+	TimeCorrection = 1,
+	Max = 2
 };
 
 
@@ -25,9 +26,58 @@ class UTimedDataMonitorEditorSettings : public UObject
 	GENERATED_BODY()
 
 public:
-	UPROPERTY(Config, EditAnywhere, Category ="Timed Data Monitor")
+	/** Option to use when calibrating from the UI. */
+	UPROPERTY(Config, EditAnywhere, Category = "Calibration")
+	FTimedDataMonitorCalibrationParameters CalibrationSettings;
+
+	/** Option to use when calibrating from the UI. */
+	UPROPERTY(Config, EditAnywhere, Category = "Calibration")
+	FTimedDataMonitorTimeCorrectionParameters TimeCorrectionSettings;
+
+	/** At which speed we should update the Timed Data Monitor UI. */
+	UPROPERTY(Config, EditAnywhere, Category ="UI", meta=(ClampMin=0.0f))
 	float RefreshRate = 0.2f;
 
+	UPROPERTY(Config, AdvancedDisplay, EditAnywhere, Category = "UI", meta = (InlineEditConditionToggle = true))
+	bool bOverrideNumberOfStandardDeviationToShow;
+
+	/**
+	 * When displaying the buffer widget, how many STD should we show.
+	 * By default, it will show the value used for Calibration or Time Correction.
+	 */
+	UPROPERTY(Config, AdvancedDisplay, EditAnywhere, Category = "UI", meta = (ClampMin = 0, ClampMax = 5, EditCondition = "bOverrideNumberOfStandardDeviationToShow"))
+	int32 OverrideNumberOfStandardDeviationToShow = 5;
+
 	UPROPERTY(Config)
 	ETimedDataMonitorEditorCalibrationType LastCalibrationType = ETimedDataMonitorEditorCalibrationType::CalibrateWithTimecode;
+
+	/** The Reset button will reset the buffer errors count. */
+	UPROPERTY(Config, EditAnywhere, Category = "UI|Reset")
+	bool bResetBufferStatEnabled = true;
+
+	/** The Reset button will clear the messages list. */
+	UPROPERTY(Config, EditAnywhere, Category = "UI|Reset")
+	bool bClearMessageEnabled = true;
+
+	/** The Reset button will set the evaluation time of all input to 0. */
+	UPROPERTY(Config, EditAnywhere, Category = "UI|Reset")
+	bool bResetEvaluationTimeEnabled = false;
+
+public:
+	int32 GetNumberOfStandardDeviationForUI() const
+	{
+		if (bOverrideNumberOfStandardDeviationToShow)
+		{
+			return OverrideNumberOfStandardDeviationToShow;
+		}
+		else if (LastCalibrationType == ETimedDataMonitorEditorCalibrationType::CalibrateWithTimecode && CalibrationSettings.bUseStandardDeviation)
+		{
+			return CalibrationSettings.NumberOfStandardDeviation;
+		}
+		else if (LastCalibrationType == ETimedDataMonitorEditorCalibrationType::TimeCorrection && TimeCorrectionSettings.bUseStandardDeviation)
+		{
+			return TimeCorrectionSettings.NumberOfStandardDeviation;
+		}
+		return 0;
+	}
 };
diff --git a/Engine/Plugins/VirtualProduction/TimedDataMonitor/Source/TimedDataMonitorEditor/Private/TimedDataMonitorEditorStyle.cpp b/Engine/Plugins/VirtualProduction/TimedDataMonitor/Source/TimedDataMonitorEditor/Private/TimedDataMonitorEditorStyle.cpp
new file mode 100644
--- /dev/null
+++ b/Engine/Plugins/VirtualProduction/TimedDataMonitor/Source/TimedDataMonitorEditor/Private/TimedDataMonitorEditorStyle.cpp
@@ -0,0 +1,158 @@
+// Copyright Epic Games, Inc. All Rights Reserved.
+
+#include "TimedDataMonitorEditorStyle.h"
+
+#include "EditorStyleSet.h"
+#include "Interfaces/IPluginManager.h"
+#include "Styling/SlateTypes.h"
+#include "Styling/CoreStyle.h"
+
+const FName FTimedDataMonitorEditorStyle::NAME_TimecodeBrush = "Img.Timecode.Small";
+const FName FTimedDataMonitorEditorStyle::NAME_PlatformTimeBrush = "Img.PlatformTime.Small";
+const FName FTimedDataMonitorEditorStyle::NAME_NoEvaluationBrush = "Img.NoEvaluation.Small";
+
+
+FTimedDataMonitorEditorStyle::FTimedDataMonitorEditorStyle()
+	: FSlateStyleSet("TimedDataSourceEditorStyle")
+{
+	const FVector2D Icon16x16(16.0f, 16.0f);
+	const FVector2D Icon24x24(24.0f, 24.0f);
+
+	FString ContentDir = IPluginManager::Get().FindPlugin(TEXT("TimedDataMonitor"))->GetContentDir();
+	SetContentRoot(ContentDir);
+	SetCoreContentRoot(FPaths::EngineContentDir() / TEXT("Slate"));
+
+	// CheckBox
+	{
+		FSlateImageBrush SwitchOn = FSlateImageBrush(RootToContentDir(TEXT("Widgets/Switch_ON.png")), FVector2D(28.f, 14.f));
+		FSlateImageBrush SwitchOff = FSlateImageBrush(RootToContentDir(TEXT("Widgets/Switch_OFF.png")), FVector2D(28.f, 14.f));
+		FSlateImageBrush SwitchUndeterminded = FSlateImageBrush(RootToContentDir(TEXT("Widgets/Switch_Undetermined.png")), FVector2D(28.f, 14.f));
+
+		FCheckBoxStyle SwitchStyle = FCheckBoxStyle()
+			.SetForegroundColor(FLinearColor::White)
+			.SetUncheckedImage(SwitchOff)
+			.SetUncheckedHoveredImage(SwitchOff)
+			.SetUncheckedPressedImage(SwitchOff)
+			.SetUndeterminedImage(SwitchUndeterminded)
+			.SetUndeterminedHoveredImage(SwitchUndeterminded)
+			.SetUndeterminedPressedImage(SwitchUndeterminded)
+			.SetCheckedImage(SwitchOn)
+			.SetCheckedHoveredImage(SwitchOn)
+			.SetCheckedPressedImage(SwitchOn)
+			.SetPadding(FMargin(0, 0, 0, 1));
+		Set("CheckBox.Enable", SwitchStyle);
+	}
+
+	// brush
+	{
+		Set("Brush.White", new FSlateColorBrush(FLinearColor::White));
+	}
+
+	// images
+	{
+		Set(NAME_TimecodeBrush, new FSlateImageBrush(RootToContentDir(TEXT("Widgets/Timecode_16x.png")), Icon16x16));
+		Set(NAME_PlatformTimeBrush, new FSlateImageBrush(RootToContentDir(TEXT("Widgets/Time_16x.png")), Icon16x16));
+		Set(NAME_NoEvaluationBrush, new FSlateImageBrush(RootToContentDir(TEXT("Widgets/NoEvaluation_16x.png")), Icon16x16));
+		Set("Img.TimedDataMonitor.Small", new FSlateImageBrush(RootToContentDir(TEXT("Widgets/TimedDataMonitor_16x.png")), Icon16x16));
+
+		Set("Img.BufferVisualization", new FSlateImageBrush(RootToContentDir(TEXT("Widgets/BufferVisualization_24x.png")), Icon24x24));
+		Set("Img.Calibration", new FSlateImageBrush(RootToContentDir(TEXT("Widgets/Calibration_24x.png")), Icon24x24));
+		Set("Img.TimeCorrection", new FSlateImageBrush(RootToContentDir(TEXT("Widgets/TimeCorrection_24x.png")), Icon24x24));
+		Set("Img.Edit", new FSlateImageBrush(RootToContentDir(TEXT("Widgets/Edit_24x.png")), Icon24x24));
+		
+	}
+
+	// font
+	{
+		Set("Font.Regular", FCoreStyle::GetDefaultFontStyle("Regular", FCoreStyle::RegularTextSize));
+		Set("Font.Large", FCoreStyle::GetDefaultFontStyle("Regular", 12));
+	}
+
+	// text block
+	{
+		FTextBlockStyle NormalText = FCoreStyle::Get().GetWidgetStyle<FTextBlockStyle>("NormalText");
+		Set("TextBlock.Regular", FTextBlockStyle(NormalText).SetFont(GetFontStyle("Font.Regular")));
+		Set("TextBlock.Large", FTextBlockStyle(NormalText).SetFont(GetFontStyle("Font.Large")));
+	}
+
+	// TableView
+	{
+		const FTableRowStyle& DefaultTableRow = FCoreStyle::Get().GetWidgetStyle<FTableRowStyle>("TableView.Row");
+		Set("TableView.Child", FTableRowStyle(DefaultTableRow)
+			.SetEvenRowBackgroundHoveredBrush(DefaultTableRow.InactiveBrush)
+			.SetOddRowBackgroundHoveredBrush(DefaultTableRow.InactiveBrush)
+			.SetActiveHoveredBrush(DefaultTableRow.InactiveBrush)
+			.SetInactiveHoveredBrush(DefaultTableRow.InactiveBrush)
+			.SetEvenRowBackgroundBrush(DefaultTableRow.InactiveBrush)
+			.SetOddRowBackgroundBrush(DefaultTableRow.InactiveBrush)
+			.SetActiveBrush(DefaultTableRow.InactiveBrush)
+			.SetInactiveBrush(DefaultTableRow.InactiveBrush)
+		);
+	}
+
+	// ComboButton
+	{
+		FComboButtonStyle SectionComboButton = FComboButtonStyle()
+			.SetButtonStyle(
+				FButtonStyle()
+				.SetNormal(FSlateNoResource())
+				.SetHovered(FSlateNoResource())
+				.SetPressed(FSlateNoResource())
+				.SetNormalPadding(FMargin(0, 0, 0, 0))
+				.SetPressedPadding(FMargin(0, 1, 0, 0))
+			)
+			.SetDownArrowImage(FSlateNoResource())
+			.SetMenuBorderBrush(FSlateNoResource());
+		SectionComboButton.UnlinkColors();
+		Set("FlatComboButton", SectionComboButton);
+
+		const FCheckBoxStyle& ToggleButtonStyle = FEditorStyle::Get().GetWidgetStyle<FCheckBoxStyle>("ToggleButtonCheckbox");
+		FComboButtonStyle ToggleComboButton = FComboButtonStyle()
+			.SetButtonStyle(
+				FButtonStyle()
+				.SetNormal(ToggleButtonStyle.UncheckedImage)
+				.SetHovered(ToggleButtonStyle.UncheckedHoveredImage)
+				.SetPressed(ToggleButtonStyle.UncheckedPressedImage)
+				.SetNormalPadding(FMargin(0, 0, 0, 0))
+				.SetPressedPadding(FMargin(0, 1, 0, 0))
+			)
+			.SetDownArrowImage(FSlateNoResource())
+			.SetMenuBorderBrush(FSlateNoResource());
+		Set("ToggleComboButton", ToggleComboButton);
+	}
+
+	// Button
+	{
+		FButtonStyle FlatButton = FButtonStyle()
+			.SetNormal(FSlateBoxBrush(RootToContentDir("Common/ButtonHoverHint.png"), FMargin(4 / 16.0f), FLinearColor(1, 1, 1, 0.15f)))
+			.SetHovered(FSlateBoxBrush(RootToContentDir("Common/ButtonHoverHint.png"), FMargin(4 / 16.0f), FLinearColor(1, 1, 1, 0.25f)))
+			.SetPressed(FSlateBoxBrush(RootToContentDir("Common/ButtonHoverHint.png"), FMargin(4 / 16.0f), FLinearColor(1, 1, 1, 0.30f)))
+			.SetNormalPadding(FMargin(0, 0, 0, 1))
+			.SetPressedPadding(FMargin(0, 1, 0, 0));
+		Set("FlatButton", FlatButton);
+
+		const FCheckBoxStyle& ToggleButtonStyle = FEditorStyle::Get().GetWidgetStyle<FCheckBoxStyle>("ToggleButtonCheckbox");
+		FButtonStyle ToggleButton = FButtonStyle()
+			.SetNormal(ToggleButtonStyle.UncheckedImage)
+			.SetHovered(ToggleButtonStyle.UncheckedHoveredImage)
+			.SetPressed(ToggleButtonStyle.UncheckedPressedImage)
+			.SetNormalPadding(FMargin(0, 0, 0, 0))
+			.SetPressedPadding(FMargin(0, 1, 0, 0));
+		Set("ToggleButton", ToggleButton);
+	}
+	
+	FSlateStyleRegistry::RegisterSlateStyle(*this);
+}
+
+FTimedDataMonitorEditorStyle::~FTimedDataMonitorEditorStyle()
+{
+	FSlateStyleRegistry::UnRegisterSlateStyle(*this);
+}
+
+FTimedDataMonitorEditorStyle& FTimedDataMonitorEditorStyle::Get()
+{
+	static FTimedDataMonitorEditorStyle Inst;
+	return Inst;
+}
+
+
diff --git a/Engine/Plugins/VirtualProduction/TimedDataMonitor/Source/TimedDataMonitorEditor/TimedDataMonitorEditor.build.cs b/Engine/Plugins/VirtualProduction/TimedDataMonitor/Source/TimedDataMonitorEditor/TimedDataMonitorEditor.build.cs
--- a/Engine/Plugins/VirtualProduction/TimedDataMonitor/Source/TimedDataMonitorEditor/TimedDataMonitorEditor.build.cs
+++ b/Engine/Plugins/VirtualProduction/TimedDataMonitor/Source/TimedDataMonitorEditor/TimedDataMonitorEditor.build.cs
@@ -14,9 +14,10 @@ public class TimedDataMonitorEditor : ModuleRules
 				"EditorStyle",
 				"InputCore",
 				"MessageLog",
+				"Projects",
 				"Settings",
-				"SlateCore",
 				"Slate",
+				"SlateCore",
 				"TimedDataMonitor",
 				"TimeManagement",
 				"TimeManagementEditor",
diff --git a/Engine/Source/Runtime/LiveLinkInterface/Private/LiveLinkSourceSettings.cpp b/Engine/Source/Runtime/LiveLinkInterface/Private/LiveLinkSourceSettings.cpp
--- a/Engine/Source/Runtime/LiveLinkInterface/Private/LiveLinkSourceSettings.cpp
+++ b/Engine/Source/Runtime/LiveLinkInterface/Private/LiveLinkSourceSettings.cpp
@@ -17,13 +17,15 @@ bool ULiveLinkSourceSettings::CanEditChange(const FProperty* InProperty) const
 	{
 		if (InProperty->GetFName() == GET_MEMBER_NAME_CHECKED(FLiveLinkSourceBufferManagementSettings, TimecodeFrameOffset)
 			|| InProperty->GetFName() == GET_MEMBER_NAME_CHECKED(FLiveLinkSourceBufferManagementSettings, TimecodeFrameRate)
-			|| InProperty->GetFName() == GET_MEMBER_NAME_CHECKED(FLiveLinkSourceBufferManagementSettings, ValidTimecodeFrame))
+			|| InProperty->GetFName() == GET_MEMBER_NAME_CHECKED(FLiveLinkSourceBufferManagementSettings, ValidTimecodeFrame)
+			|| InProperty->GetFName() == GET_MEMBER_NAME_CHECKED(FLiveLinkSourceBufferManagementSettings, bValidTimecodeFrameEnabled))
 		{
 			return Mode == ELiveLinkSourceMode::Timecode;
 		}
 
 		if (InProperty->GetFName() == GET_MEMBER_NAME_CHECKED(FLiveLinkSourceBufferManagementSettings, ValidEngineTime)
-			|| InProperty->GetFName() == GET_MEMBER_NAME_CHECKED(FLiveLinkSourceBufferManagementSettings, EngineTimeOffset))
+			|| InProperty->GetFName() == GET_MEMBER_NAME_CHECKED(FLiveLinkSourceBufferManagementSettings, EngineTimeOffset)
+			|| InProperty->GetFName() == GET_MEMBER_NAME_CHECKED(FLiveLinkSourceBufferManagementSettings, bValidEngineTimeEnabled))
 		{
 			return Mode == ELiveLinkSourceMode::EngineTime;
 		}
diff --git a/Engine/Source/Runtime/LiveLinkInterface/Public/LiveLinkSourceSettings.h b/Engine/Source/Runtime/LiveLinkInterface/Public/LiveLinkSourceSettings.h
--- a/Engine/Source/Runtime/LiveLinkInterface/Public/LiveLinkSourceSettings.h
+++ b/Engine/Source/Runtime/LiveLinkInterface/Public/LiveLinkSourceSettings.h
@@ -39,8 +39,12 @@ struct FLiveLinkSourceBufferManagementSettings
 {
 	GENERATED_BODY()
 
+	/** Enabled the ValidEngineTime setting. */
+	UPROPERTY(EditAnywhere, Category = "Settings", meta=(InlineEditConditionToggle=true))
+	bool bValidEngineTimeEnabled = false;
+
 	/** If the frame is older than ValidTime, remove it from the buffer list (in seconds). */
-	UPROPERTY(EditAnywhere, Category = "Settings", meta=(ForceUnits=s, ClampMin=0.5))
+	UPROPERTY(EditAnywhere, Category = "Settings", meta=(ForceUnits=s, ClampMin=0.5, EditCondition="bValidEngineTimeEnabled"))
 	float ValidEngineTime = 1.0f;
 
 	/** When evaluating with time: how far back from current time should we read the buffer (in seconds) */
@@ -64,7 +68,11 @@ struct FLiveLinkSourceBufferManagementSettings
 	FFrameRate SourceTimecodeFrameRate = { 24, 1 };
 
 	/** If the frame timecode is older than ValidTimecodeFrame, remove it from the buffer list (in TimecodeFrameRate). */
-	UPROPERTY(EditAnywhere, Category = "Settings", meta=(ClampMin=1))
+	UPROPERTY(EditAnywhere, Category = "Settings", meta=(InlineEditConditionToggle=true))
+	bool bValidTimecodeFrameEnabled = false;
+
+	/** If the frame timecode is older than ValidTimecodeFrame, remove it from the buffer list (in TimecodeFrameRate). */
+	UPROPERTY(EditAnywhere, Category = "Settings", meta=(ClampMin=1, EditCondition="bValidTimecodeFrameEnabled"))
 	int32 ValidTimecodeFrame = 30;
 
 	/** When evaluating with timecode: how far back from current timecode should we read the buffer (in TimecodeFrameRate). */
diff --git a/Engine/Source/Runtime/TimeManagement/Public/ITimedDataInput.h b/Engine/Source/Runtime/TimeManagement/Public/ITimedDataInput.h
--- a/Engine/Source/Runtime/TimeManagement/Public/ITimedDataInput.h
+++ b/Engine/Source/Runtime/TimeManagement/Public/ITimedDataInput.h
@@ -57,6 +57,8 @@ struct FTimedDataChannelSampleTime
 	double PlatformSecond = 0.0;
 	/** Timecode value of the sample */
 	FQualifiedFrameTime Timecode;
+
+	double AsSeconds(ETimedDataInputEvaluationType EvaluationType) const { return EvaluationType == ETimedDataInputEvaluationType::Timecode ? Timecode.AsSeconds() : PlatformSecond; }
 };
 
 
@@ -120,6 +122,12 @@ public:
 
 	/** If the input does supported it, set the size of the buffer used by the input. */
 	virtual void SetDataBufferSize(int32 BufferSize) { }
+	
+	/** Add a channel belonging to this input */
+	virtual void AddChannel(ITimedDataInputChannel* Channel) = 0;
+
+	/** Remove channel from the input */
+	virtual void RemoveChannel(ITimedDataInputChannel* Channel) = 0;
 
 #if WITH_EDITOR
 	/** Get the icon that represent the input. */
diff --git a/Engine/Source/Runtime/TimeManagement/Public/TimedDataInputCollection.h b/Engine/Source/Runtime/TimeManagement/Public/TimedDataInputCollection.h
--- a/Engine/Source/Runtime/TimeManagement/Public/TimedDataInputCollection.h
+++ b/Engine/Source/Runtime/TimeManagement/Public/TimedDataInputCollection.h
@@ -41,11 +41,11 @@ public:
 	}
 	
 	/** Add an input channel to the collection. */
-	void Add(ITimedDataInputChannel* Input)
+	void Add(ITimedDataInputChannel* Channel)
 	{
-		if (Input && !Groups.Contains(Input))
+		if (Channel && !Channels.Contains(Channel))
 		{
-			Groups.Add(Input);
+			Channels.Add(Channel);
 			OnCollectionChanged().Broadcast();
 		}
 	}
@@ -53,7 +53,7 @@ public:
 	/** Remove an input channel from the collection. */
 	void Remove(ITimedDataInputChannel* Input)
 	{
-		if (Groups.RemoveSingleSwap(Input) > 0)
+		if (Channels.RemoveSingleSwap(Input) > 0)
 		{
 			OnCollectionChanged().Broadcast();
 		}
@@ -62,7 +62,7 @@ public:
 	/** The list of input channels from the collection. */
 	const TArray<ITimedDataInputChannel*>& GetChannels() const
 	{
-		return Groups;
+		return Channels;
 	}
 
 	/** When an element is added or removed to the collection. */
@@ -74,5 +74,5 @@ public:
 private:
 	FSimpleMulticastDelegate CollectionChanged;
 	TArray<ITimedDataInput*> Inputs;
-	TArray<ITimedDataInputChannel*> Groups;
+	TArray<ITimedDataInputChannel*> Channels;
 };
\ No newline at end of file
