Merging CL#11565130 from Dev-VirtualProduction

- Changed stats to use exponential moving avg and variance to better track changes (cvar to change weights) and added feedforward to track faster when offsets change
- Fixed bug with IsInRange sometime failing
- Adding file logging for timed data channels stats

diff --git a/Engine/Plugins/Media/MediaIOFramework/Source/MediaIOCore/Private/Player/MediaIOCorePlayerBase.cpp b/Engine/Plugins/Media/MediaIOFramework/Source/MediaIOCore/Private/Player/MediaIOCorePlayerBase.cpp
--- a/Engine/Plugins/Media/MediaIOFramework/Source/MediaIOCore/Private/Player/MediaIOCorePlayerBase.cpp
+++ b/Engine/Plugins/Media/MediaIOFramework/Source/MediaIOCore/Private/Player/MediaIOCorePlayerBase.cpp
@@ -553,7 +553,7 @@ void FMediaIOCorePlayerBase::SetEvaluationOffsetInSeconds(double Offset)
 	{
 		//Media doesn't support subframes playback (interpolation between frames) so offsets are always in full frame
 		const double FrameOffset = ITimedDataInput::ConvertSecondOffsetInFrameOffset(Offset, VideoFrameRate);
-		FrameDelay = FMath::FloorToInt(FrameOffset);
+		FrameDelay = FMath::CeilToInt(FrameOffset);
 	}
 	else
 	{
diff --git a/Engine/Plugins/VirtualProduction/TimedDataMonitor/Source/TimedDataMonitor/Private/TimedDataMonitorSubsystem.cpp b/Engine/Plugins/VirtualProduction/TimedDataMonitor/Source/TimedDataMonitor/Private/TimedDataMonitorSubsystem.cpp
--- a/Engine/Plugins/VirtualProduction/TimedDataMonitor/Source/TimedDataMonitor/Private/TimedDataMonitorSubsystem.cpp
+++ b/Engine/Plugins/VirtualProduction/TimedDataMonitor/Source/TimedDataMonitor/Private/TimedDataMonitorSubsystem.cpp
@@ -16,6 +16,29 @@
 
 static TAutoConsoleVariable<bool> CVarEnableTimedDataMonitorSubsystemStats(TEXT("TimedDataMonitor.EnableStatUpdate"), 1, TEXT("Enable calculating evaluation statistics of all registered channels."));
 
+#if !(UE_BUILD_SHIPPING || UE_BUILD_TEST)
+bool bIsTimedDataStatFileLoggingStarted = false;
+static FAutoConsoleCommand TimedDataStartFileLoggingCmd(
+	TEXT("TimedDataMonitor.StartFileLogging"),
+	TEXT("Starts logging data for each channel for each frame to be logged in a file."),
+	FConsoleCommandDelegate::CreateLambda([]() { bIsTimedDataStatFileLoggingStarted = true; })
+);
+
+static FAutoConsoleCommand TimedDataDisableFileLoggingCmd(
+	TEXT("TimedDataMonitor.StopFileLogging"),
+	TEXT("Stops logging data and dump everything to a file."),
+	FConsoleCommandDelegate::CreateLambda([]() { bIsTimedDataStatFileLoggingStarted = false; })
+);
+#endif //!(UE_BUILD_SHIPPING || UE_BUILD_TEST)
+
+TAutoConsoleVariable<float> CVarTimedDataStatisticsWeight(
+	TEXT("TimedDataMonitor.Statistics.Weight"),
+	0.1f,
+	TEXT("The weight used when tracking Mean and Variance of samples distance to evaluation time. Number closer to 1 will give more weight to latest values and won't filter out noise."),
+	ECVF_Default
+);
+
+
 #define LOCTEXT_NAMESPACE "TimedDataMonitorSubsystem"
 
 /**
@@ -61,14 +84,20 @@ void UTimedDataMonitorSubsystem::FTimeDataChannelItem::ResetValue()
  */
 void UTimedDataMonitorSubsystem::Initialize(FSubsystemCollectionBase& Collection)
 {
+	Super::Initialize(Collection);
+	
 	bRequestSourceListRebuilt = true;
 	ITimeManagementModule::Get().GetTimedDataInputCollection().OnCollectionChanged().AddUObject(this, &UTimedDataMonitorSubsystem::OnTimedDataSourceCollectionChanged);
-	Super::Initialize(Collection);
+	FCoreDelegates::OnEndFrame.AddUObject(this, &UTimedDataMonitorSubsystem::EndFrameCallback);
+	FCoreDelegates::OnBeginFrame.AddUObject(this, &UTimedDataMonitorSubsystem::BeginFrameCallback);
 }
 
 
 void UTimedDataMonitorSubsystem::Deinitialize()
 {
+	FCoreDelegates::OnBeginFrame.RemoveAll(this);
+	FCoreDelegates::OnEndFrame.RemoveAll(this);
+
 	if (ITimeManagementModule::IsAvailable())
 	{
 		ITimeManagementModule::Get().GetTimedDataInputCollection().OnCollectionChanged().RemoveAll(this);
@@ -84,6 +113,32 @@ void UTimedDataMonitorSubsystem::Deinitialize()
 }
 
 
+void UTimedDataMonitorSubsystem::BeginFrameCallback()
+{
+	UTimecodeProvider* CurrentTimecodeProvider = GEngine->GetTimecodeProvider();
+	const TOptional<FQualifiedFrameTime> CurrentFrameTimeOffsetted = FApp::GetCurrentFrameTime();
+	if (CurrentTimecodeProvider == nullptr
+		|| CurrentTimecodeProvider->GetSynchronizationState() != ETimecodeProviderSynchronizationState::Synchronized
+		|| !CurrentFrameTimeOffsetted.IsSet())
+	{
+		CachedTimecodeProviderFrameDelayInSeconds = 0.0f;
+	}
+	else
+	{
+		CachedTimecodeProviderFrameDelayInSeconds = CurrentTimecodeProvider->GetFrameRate().AsInterval() * CurrentTimecodeProvider->FrameDelay;
+	}
+}
+
+void UTimedDataMonitorSubsystem::EndFrameCallback()
+{
+	const bool bUpdateStats = CVarEnableTimedDataMonitorSubsystemStats.GetValueOnGameThread();
+	if (bUpdateStats)
+	{
+		UpdateStatFileLoggingState();
+		UpdateEvaluationStatistics();
+	}
+}
+
 ITimedDataInput* UTimedDataMonitorSubsystem::GetTimedDataInput(const FTimedDataMonitorInputIdentifier& Identifier)
 {
 	BuildSourcesListIfNeeded();
@@ -110,26 +165,46 @@ ITimedDataInputChannel* UTimedDataMonitorSubsystem::GetTimedDataChannel(const FT
 }
 
 
- double UTimedDataMonitorSubsystem::GetEvaluationTime(ETimedDataInputEvaluationType EvaluationType)
- {
-	 double Result = 0.0;
-	 switch (EvaluationType)
-	 {
-	 case ETimedDataInputEvaluationType::Timecode:
-		 if (FApp::GetCurrentFrameTime().IsSet())
-		 {
-			 Result = FApp::GetCurrentFrameTime().GetValue().AsSeconds();
-		 }
-		 break;
-	 case ETimedDataInputEvaluationType::PlatformTime:
-		 Result = FApp::GetCurrentTime();
-		 break;
-	 case ETimedDataInputEvaluationType::None:
-	 default:
-		 break;
-	 }
-	 return Result;
- }
+float UTimedDataMonitorSubsystem::GetEvaluationTimeOffsetInSeconds(ETimedDataInputEvaluationType EvaluationType)
+{
+	float Result = 0.0f;
+
+	switch (EvaluationType)
+	{
+	case ETimedDataInputEvaluationType::Timecode:
+	{
+		Result = CachedTimecodeProviderFrameDelayInSeconds;
+		break;
+	}
+	case ETimedDataInputEvaluationType::PlatformTime:
+	case ETimedDataInputEvaluationType::None:
+	default:
+		break;
+	}
+
+	return Result;
+}
+
+double UTimedDataMonitorSubsystem::GetEvaluationTime(ETimedDataInputEvaluationType EvaluationType)
+{
+	double Result = 0.0;
+	switch (EvaluationType)
+	{
+	case ETimedDataInputEvaluationType::Timecode:
+		if (FApp::GetCurrentFrameTime().IsSet())
+		{
+			Result = FApp::GetCurrentFrameTime().GetValue().AsSeconds();
+		}
+		break;
+	case ETimedDataInputEvaluationType::PlatformTime:
+		Result = FApp::GetCurrentTime();
+		break;
+	case ETimedDataInputEvaluationType::None:
+	default:
+	break;
+	}
+	return Result;
+}
 
 
 TArray<FTimedDataMonitorInputIdentifier> UTimedDataMonitorSubsystem::GetAllInputs()
@@ -570,7 +645,7 @@ float UTimedDataMonitorSubsystem::GetInputEvaluationDistanceToNewestSampleMean(c
 		WorstNewestMean = TNumericLimits<float>::Lowest();
 		for (const FTimedDataMonitorChannelIdentifier& ChannelIdentifier : InputItem->ChannelIdentifiers)
 		{
-			WorstNewestMean = FMath::Max(ChannelMap[ChannelIdentifier].Statistics.IncrementalAverageNewestDistance, WorstNewestMean);
+			WorstNewestMean = FMath::Max(ChannelMap[ChannelIdentifier].Statistics.NewestSampleDistanceTracker.CurrentMean, WorstNewestMean);
 		}
 	}
 
@@ -588,7 +663,7 @@ float UTimedDataMonitorSubsystem::GetInputEvaluationDistanceToOldestSampleMean(c
 		WorstOldesttMean = TNumericLimits<float>::Max();
 		for (const FTimedDataMonitorChannelIdentifier& ChannelIdentifier : InputItem->ChannelIdentifiers)
 		{
-			WorstOldesttMean = FMath::Min(ChannelMap[ChannelIdentifier].Statistics.IncrementalAverageOldestDistance, WorstOldesttMean);
+			WorstOldesttMean = FMath::Min(ChannelMap[ChannelIdentifier].Statistics.OldestSampleDistanceTracker.CurrentMean, WorstOldesttMean);
 		}
 	}
 
@@ -605,7 +680,7 @@ float UTimedDataMonitorSubsystem::GetInputEvaluationDistanceToNewestSampleStanda
 	{
 		for (const FTimedDataMonitorChannelIdentifier& ChannelIdentifier : InputItem->ChannelIdentifiers)
 		{
-			WorstNewestSSD = FMath::Max(ChannelMap[ChannelIdentifier].Statistics.DistanceToNewestSTD, WorstNewestSSD);
+			WorstNewestSSD = FMath::Max(ChannelMap[ChannelIdentifier].Statistics.NewestSampleDistanceTracker.CurrentSTD, WorstNewestSSD);
 		}
 	}
 
@@ -622,7 +697,7 @@ float UTimedDataMonitorSubsystem::GetInputEvaluationDistanceToOldestSampleStanda
 	{
 		for (const FTimedDataMonitorChannelIdentifier& ChannelIdentifier : InputItem->ChannelIdentifiers)
 		{
-			WorstOldestSSD = FMath::Max(ChannelMap[ChannelIdentifier].Statistics.DistanceToOldestSTD, WorstOldestSSD);
+			WorstOldestSSD = FMath::Max(ChannelMap[ChannelIdentifier].Statistics.OldestSampleDistanceTracker.CurrentSTD, WorstOldestSSD);
 		}
 	}
 
@@ -724,7 +799,7 @@ ETimedDataMonitorEvaluationState UTimedDataMonitorSubsystem::GetChannelEvaluatio
 		const double NewstedSampleTime = SourceItem->Channel->GetNewestDataTime().AsSeconds(EvaluationType);
 		const double EvaluationTime = GetEvaluationTime(EvaluationType);
 		const double OffsettedEvaluationTime = EvaluationTime - EvaluationOffset;
-		bool bIsInRange = (OffsettedEvaluationTime >= OldestSampleTime) && (OffsettedEvaluationTime <= NewstedSampleTime);
+		const bool bIsInRange = (FMath::IsNearlyEqual(OffsettedEvaluationTime, OldestSampleTime) || OffsettedEvaluationTime >= OldestSampleTime) && (FMath::IsNearlyEqual(OffsettedEvaluationTime, NewstedSampleTime) ||  OffsettedEvaluationTime <= NewstedSampleTime);
 		return bIsInRange ? ETimedDataMonitorEvaluationState::InsideRange : ETimedDataMonitorEvaluationState::OutsideRange;
 	}
 
@@ -847,7 +922,7 @@ float UTimedDataMonitorSubsystem::GetChannelEvaluationDistanceToNewestSampleMean
 
 	if (const FTimeDataChannelItem* SourceItem = ChannelMap.Find(Identifier))
 	{
-		return SourceItem->Statistics.IncrementalAverageNewestDistance;
+		return SourceItem->Statistics.NewestSampleDistanceTracker.CurrentMean;
 	}
 
 	return 0.0f;
@@ -859,7 +934,7 @@ float UTimedDataMonitorSubsystem::GetChannelEvaluationDistanceToOldestSampleMean
 
 	if (const FTimeDataChannelItem* SourceItem = ChannelMap.Find(Identifier))
 	{
-		return SourceItem->Statistics.IncrementalAverageOldestDistance;
+		return SourceItem->Statistics.OldestSampleDistanceTracker.CurrentMean;
 	}
 
 	return 0.0f;
@@ -871,7 +946,7 @@ float UTimedDataMonitorSubsystem::GetChannelEvaluationDistanceToNewestSampleStan
 
 	if (const FTimeDataChannelItem* SourceItem = ChannelMap.Find(Identifier))
 	{
-		return SourceItem->Statistics.DistanceToNewestSTD;
+		return SourceItem->Statistics.NewestSampleDistanceTracker.CurrentSTD;
 	}
 
 	return 0.0f;
@@ -883,7 +958,7 @@ float UTimedDataMonitorSubsystem::GetChannelEvaluationDistanceToOldestSampleStan
 
 	if (const FTimeDataChannelItem* SourceItem = ChannelMap.Find(Identifier))
 	{
-		return SourceItem->Statistics.DistanceToOldestSTD;
+		return SourceItem->Statistics.OldestSampleDistanceTracker.CurrentSTD;
 	}
 
 	return 0.0f;
@@ -1063,22 +1138,6 @@ void UTimedDataMonitorSubsystem::OnTimedDataSourceCollectionChanged()
 }
 
 
-void UTimedDataMonitorSubsystem::Tick(float DeltaTime)
-{
-	const bool bUpdateStats = CVarEnableTimedDataMonitorSubsystemStats.GetValueOnGameThread();
-	if (bUpdateStats)
-	{
-		UpdateEvaluationStatistics();
-	}
-}
-
-
-TStatId UTimedDataMonitorSubsystem::GetStatId() const
-{
-	RETURN_QUICK_DECLARE_CYCLE_STAT(UTimedDataMonitorSubsystem, STATGROUP_Tickables);
-}
-
-
 void UTimedDataMonitorSubsystem::UpdateEvaluationStatistics()
 {
 	BuildSourcesListIfNeeded();
@@ -1090,54 +1149,160 @@ void UTimedDataMonitorSubsystem::UpdateEvaluationStatistics()
 			FTimedDataInputEvaluationData Data;
 			GetChannelLastEvaluationDataStat(Item.Key, Data);
 
+			//Update cached settings to accelerate response of mean tracker
+			const ITimedDataInput* Input = InputMap[Item.Value.InputIdentifier].Input;
+			const ETimedDataInputEvaluationType EvaluationType = Input->GetEvaluationType();
+
+			const float EvalutionOffset = GetEvaluationTimeOffsetInSeconds(EvaluationType);
+			const float TotalTimeOffset = GetInputEvaluationOffsetInSeconds(Item.Value.InputIdentifier) + EvalutionOffset;
+			const int32 BufferSize = Input->IsDataBufferSizeControlledByInput() ? Input->GetDataBufferSize() : Item.Value.Channel->GetDataBufferSize();
+			Item.Value.Statistics.CacheSettings(EvaluationType, TotalTimeOffset, BufferSize);
+
+			//Update statistics for each channel
 			Item.Value.Statistics.Update(Data.DistanceToOldestSampleSeconds, Data.DistanceToNewestSampleSeconds);
+
+			AddStatisticLogEntry(Item);
 		}
 	}
 }
 
-void FTimedDataChannelEvaluationStatistics::Update(float DistanceToOldest, float DistanceToNewest)
+void UTimedDataMonitorSubsystem::UpdateStatFileLoggingState()
 {
-	//Compute running average and variance based on Welford's algorithm 
-	//https://en.wikipedia.org/wiki/Algorithms_for_calculating_variance#Welford's_online_algorithm
+#if !(UE_BUILD_SHIPPING || UE_BUILD_TEST)
+	if (bHasStatFileLoggingStarted)
+	{
+		if (!bIsTimedDataStatFileLoggingStarted)
+		{
+			const FDateTime NowTime;
+			const FString CurrentTimeString = FDateTime::Now().ToString();
+			const FString OutputDirectory = FPaths::Combine(*FPaths::ProjectDir(), TEXT("Saved"), TEXT("TimedDataMonitor"));
 
-	//Update sample count to include this new one
-	++SampleCount;
+			//Dump each channel to a file
+			for (TPair<FTimedDataMonitorChannelIdentifier, FChannelStatisticLogging>& Item : StatLoggingMap)
+			{
+				const FString Filename = FString::Printf(TEXT("%s_%s.csv"), *CurrentTimeString, *Item.Value.ChannelName);
+				const FString FullyQualifiedFileNameFilename = FPaths::Combine(OutputDirectory, Filename);
+				FFileHelper::SaveStringArrayToFile(Item.Value.Entries, *FullyQualifiedFileNameFilename);
+			}
+
+			bHasStatFileLoggingStarted = false;
+		}
+	}
+	else
+	{
+		if (bIsTimedDataStatFileLoggingStarted)
+		{
+			bHasStatFileLoggingStarted = true;
+			StatLoggingMap.Empty();
+		}
+	}
+#endif //!(UE_BUILD_SHIPPING || UE_BUILD_TEST)
+}
+
+void UTimedDataMonitorSubsystem::AddStatisticLogEntry(const TPair<FTimedDataMonitorChannelIdentifier, FTimeDataChannelItem>& ChannelEntry)
+{
+#if !(UE_BUILD_SHIPPING || UE_BUILD_TEST)
+	if (bHasStatFileLoggingStarted)
+	{
+		const ITimedDataInput* Input = InputMap[ChannelEntry.Value.InputIdentifier].Input;
+		check(Input);
+		const ETimedDataInputEvaluationType EvaluationType = Input->GetEvaluationType();
+		
+		FTimedDataInputEvaluationData Data;
+		GetChannelLastEvaluationDataStat(ChannelEntry.Key, Data);
+
+		FChannelStatisticLogging& ChannelStats = StatLoggingMap.FindOrAdd(ChannelEntry.Key);
+		if (ChannelStats.Entries.Num() <= 0)
+		{
+			ChannelStats.ChannelName = ChannelMap[ChannelEntry.Key].Channel->GetDisplayName().ToString();
+			const FString Header = TEXT("SampleCount, EvaluationTime,OldestSampleTime, OldestDistance, OldestDistanceMean, OldestDistanceVariance, NewestSampleTime, NewestDistance, NewestDistanceMean, NewestDistanceVariance");
+			ChannelStats.Entries.Emplace(Header);
+		}
 
-	//For running variance, keep previous mean distance
-	const float NewDistanceToNewestToMean1 = DistanceToNewest - IncrementalAverageNewestDistance;
-	const float NewDistanceToOldestToMean1 = DistanceToOldest - IncrementalAverageOldestDistance;
+		FString NewEntry = FString::Printf(TEXT("%d,%0.8f,%0.8f,%0.8f,%0.8f,%0.8f,%0.8f,%0.8f,%0.8f,%0.8f")
+			, ChannelEntry.Value.Statistics.NewestSampleDistanceTracker.SampleCount
+			, GetEvaluationTime(EvaluationType)
+			, ChannelEntry.Value.Channel->GetOldestDataTime().AsSeconds(EvaluationType)
+			, Data.DistanceToOldestSampleSeconds
+			, ChannelEntry.Value.Statistics.OldestSampleDistanceTracker.CurrentMean
+			, ChannelEntry.Value.Statistics.OldestSampleDistanceTracker.CurrentVariance
+			, ChannelEntry.Value.Channel->GetNewestDataTime().AsSeconds(EvaluationType)
+			, Data.DistanceToNewestSampleSeconds
+			, ChannelEntry.Value.Statistics.NewestSampleDistanceTracker.CurrentMean
+			, ChannelEntry.Value.Statistics.NewestSampleDistanceTracker.CurrentVariance
+		);
+		ChannelStats.Entries.Emplace(MoveTemp(NewEntry));
+	}
+#endif //!(UE_BUILD_SHIPPING || UE_BUILD_TEST)
+}
 
-	//Update incremental average of distances to both ends of this input buffer
-	IncrementalAverageNewestDistance = IncrementalAverageNewestDistance + (DistanceToNewest - IncrementalAverageNewestDistance) / SampleCount;
-	IncrementalAverageOldestDistance = IncrementalAverageOldestDistance + (DistanceToOldest - IncrementalAverageOldestDistance) / SampleCount;
+void FTimedDataChannelEvaluationStatistics::CacheSettings(ETimedDataInputEvaluationType EvaluationType, float TimeOffset, int32 BufferSize)
+{
+	if (EvaluationType != CachedEvaluationType)
+	{
+		OldestSampleDistanceTracker.Reset();
+		NewestSampleDistanceTracker.Reset();
+	}
 
-	//Compute sum of squares for running variance
-	const float NewDistanceToNewestToMean2 = DistanceToNewest - IncrementalAverageNewestDistance;
-	SumSquaredDistanceNewest = SumSquaredDistanceNewest + NewDistanceToNewestToMean2 * NewDistanceToNewestToMean1;
-	const float NewDistanceToOldestToMean2 = DistanceToOldest - IncrementalAverageOldestDistance;
-	SumSquaredDistanceOldest = SumSquaredDistanceOldest + NewDistanceToOldestToMean2 * NewDistanceToOldestToMean1;
+	CachedEvaluationType = EvaluationType;
 
-	//Finally compute the variance
-	IncrementalVarianceDistanceNewest = SumSquaredDistanceNewest / SampleCount;
-	IncrementalVarianceDistanceOldest = SumSquaredDistanceOldest / SampleCount;
+	//Update next tick offset to feedforward our mean/variance tracker
+	NextTickOffset = TimeOffset - CachedOffset;
+	
+	CachedOffset = TimeOffset;
+	CachedBufferSize = BufferSize;
+}
 
-	//Square root of that average gives us sigma (standard deviation)
-	DistanceToNewestSTD = FMath::Sqrt(IncrementalVarianceDistanceNewest);
-	DistanceToOldestSTD = FMath::Sqrt(IncrementalVarianceDistanceOldest);
+void FTimedDataChannelEvaluationStatistics::Update(float DistanceToOldest, float DistanceToNewest)
+{
+	//Update each stats tracker using the current offset. This only lasts one frame to anticipate stats movement
+	OldestSampleDistanceTracker.Update(DistanceToOldest, -NextTickOffset);
+	NewestSampleDistanceTracker.Update(DistanceToNewest, NextTickOffset);
 
-	LastDistanceToOldest = DistanceToOldest;
-	LastDistanceToNewest = DistanceToNewest;
+	NextTickOffset = 0.0f;
 }
 
 void FTimedDataChannelEvaluationStatistics::Reset()
+{
+	OldestSampleDistanceTracker.Reset();
+	NewestSampleDistanceTracker.Reset();
+}
+
+void FExponentialMeanVarianceTracker::Reset()
 {
 	SampleCount = 0;
-	IncrementalAverageOldestDistance = 0.0f;
-	IncrementalAverageNewestDistance = 0.0f;
-	IncrementalVarianceDistanceNewest = 0.0f;
-	IncrementalVarianceDistanceOldest = 0.0f;
-	SumSquaredDistanceNewest = 0.0f;
-	SumSquaredDistanceOldest = 0.0f;
+	CurrentMean = 0.0f;
+	CurrentVariance = 0.0f;
+	CurrentSTD = 0.0f;
+	Alpha = CVarTimedDataStatisticsWeight.GetValueOnGameThread();
 }
 
+void FExponentialMeanVarianceTracker::Update(float NewValue, float MeanOffset)
+{
+	//Exponential moving average computation based on https://en.wikipedia.org/wiki/Moving_average#Exponential_moving_average
+
+	if (SampleCount > 0)
+	{
+		const float Beta = 1.0f - Alpha;
+		const float BetaMean = Beta * (MeanOffset + CurrentMean);
+		const float MeanAddition = NewValue * Alpha;
+		const float VarianceIncrement = Alpha * FMath::Square(NewValue - (CurrentMean + MeanOffset));
+		CurrentMean = BetaMean + MeanAddition;
+		CurrentVariance = Beta * (CurrentVariance + VarianceIncrement);
+	}
+	else
+	{
+		//When starting exponential count, start at current Value to get better initial response
+		CurrentMean = NewValue;
+		CurrentVariance = 0.0f;
+	}
+
+	LastValue = NewValue;
+
+	CurrentSTD = FMath::Sqrt(CurrentVariance);
+	++SampleCount;
+}
+
+
 #undef LOCTEXT_NAMESPACE
+
diff --git a/Engine/Plugins/VirtualProduction/TimedDataMonitor/Source/TimedDataMonitor/Public/TimedDataMonitorSubsystem.h b/Engine/Plugins/VirtualProduction/TimedDataMonitor/Source/TimedDataMonitor/Public/TimedDataMonitorSubsystem.h
--- a/Engine/Plugins/VirtualProduction/TimedDataMonitor/Source/TimedDataMonitor/Public/TimedDataMonitorSubsystem.h
+++ b/Engine/Plugins/VirtualProduction/TimedDataMonitor/Source/TimedDataMonitor/Public/TimedDataMonitorSubsystem.h
@@ -34,44 +34,60 @@ enum class ETimedDataMonitorEvaluationState : uint8
 
 
 /**
- * Structure to facilitate calculating running mean and variance of evaluation distance to buffered samples for a channel
+ * Exponential running mean calculator. Gives better result than incremental running mean when parameters change
  */
-struct FTimedDataChannelEvaluationStatistics
+struct FExponentialMeanVarianceTracker
 {
-	void Update(float DistanceToOldest, float DistanceToNewest);
+	FExponentialMeanVarianceTracker()
+	{
+		Reset();
+	}
+
 	void Reset();
+	void Update(float NewValue, float MeanOffset);
 
 	/** Current number of samples used for the running average and variance */
 	int32 SampleCount = 0;
+	
+	/** Last sample value used for statistics */
+	float LastValue = 0.0f;
+	
+	/** Current running mean */
+	float CurrentMean = 0.0f;
+	
+	/** Current running variance */
+	float CurrentVariance = 0.0f;
+	
+	/** Current standard deviation */
+	float CurrentSTD = 0.0f;
+	
+	/** Weight given to new samples. A bigger value means new data will weight more in the calculation. Won't filter if data oscillates a lot. */
+	float Alpha = 0.0f;
+};
 
-	/** Last distance from evaluation time to oldest sample used for statistics */
-	float LastDistanceToOldest = 0.0f;
-
-	/** Last distance from evaluation time to newest sample used for statistics */
-	float LastDistanceToNewest = 0.0f;
-
-	/** Running average of the distance between last evaluation time and oldest sample in the buffer */
-	float IncrementalAverageOldestDistance = 0.0f;
-
-	/** Running average of the distance between last evaluation time and newest sample in the buffer */
-	float IncrementalAverageNewestDistance = 0.0f;
-
-	/** Running variance of the distance between last evaluation time and oldest sample in the buffer */
-	float IncrementalVarianceDistanceNewest = 0.0f;
-
-	/** Running variance of the distance between last evaluation time and newest sample in the buffer */
-	float IncrementalVarianceDistanceOldest = 0.0f;
+/**
+ * Structure to facilitate calculating running mean and variance of evaluation distance to buffered samples for a channel
+ */
+struct FTimedDataChannelEvaluationStatistics
+{
+	void CacheSettings(ETimedDataInputEvaluationType EvaluationType, float TimeOffset, int32 BufferSize);
+	void Update(float DistanceToOldest, float DistanceToNewest);
+	void Reset();
+	
+	FExponentialMeanVarianceTracker NewestSampleDistanceTracker;
+	FExponentialMeanVarianceTracker OldestSampleDistanceTracker;
 
-	/** Standard deviation of the distance between last evaluation time and oldest sample in the buffer */
-	float DistanceToNewestSTD = 0.0f;
+	/** Last evaluation type of our input */
+	ETimedDataInputEvaluationType CachedEvaluationType;
 
-	/** Standard deviation of the distance between last evaluation time and newest sample in the buffer */
-	float DistanceToOldestSTD = 0.0f;
+	/** Last Offset associated with this channel */
+	float CachedOffset = 0.0f;
 
-	/** Internal counters of squared distance to average to be able to compute running variance */
-	double SumSquaredDistanceNewest = 0.0;
-	double SumSquaredDistanceOldest = 0.0;
+	/** Last buffer size associated with this channel */
+	int32 CachedBufferSize = 0;
 
+	/** Offset to be applied for the next tick. Used to feedforward mean tracker */
+	float NextTickOffset = 0.0f;
 };
 
 
@@ -83,7 +99,7 @@ DECLARE_DYNAMIC_MULTICAST_DELEGATE(FTimedDataIdentifierListChangedSignature);
  * 
  */
 UCLASS()
-class TIMEDDATAMONITOR_API UTimedDataMonitorSubsystem : public UEngineSubsystem, public FTickableGameObject
+class TIMEDDATAMONITOR_API UTimedDataMonitorSubsystem : public UEngineSubsystem
 {
 	GENERATED_BODY()
 
@@ -115,13 +131,11 @@ public:
 	virtual void Deinitialize() override;
 	//~ End USubsystem implementation
 
-	//~ Begin FTickableGameObject implementation
-	virtual void Tick(float DeltaTime) override;
-	virtual TStatId GetStatId() const override;
-	virtual ETickableTickType GetTickableTickType() const override { return ETickableTickType::Always; }
-	virtual bool IsTickableWhenPaused() const override { return true; }
-	virtual bool IsTickableInEditor() const override { return true; }
-	//~ End FTickableGameObject implementation
+	/** Used to cache global (Engine TimecodeProvider) frame delay for the current frame. */
+	void BeginFrameCallback();
+
+	/** Used to update statistics once all TimedData are processed. i.e. MediaIO is processed after tickables. */
+	void EndFrameCallback();
 
 public:
 	/** Delegate of when an element is added or removed. */
@@ -138,6 +152,9 @@ public:
 	/** Get the interface for a specific channel identifier. */
 	ITimedDataInputChannel* GetTimedDataChannel(const FTimedDataMonitorChannelIdentifier& Identifier);
 
+	/** Get offset applied to global evaluation time. Only works when a Timecode Provider is used */
+	float GetEvaluationTimeOffsetInSeconds(ETimedDataInputEvaluationType EvaluationType);
+
 	/** Get the current evaluation time. */
 	static double GetEvaluationTime(ETimedDataInputEvaluationType EvaluationType);
 
@@ -348,9 +365,27 @@ private:
 	/** Update internal statistics for each enabled channel */
 	void UpdateEvaluationStatistics();
 
+	/** Starts logging stats for each channel to be dumped in a file */
+	void UpdateStatFileLoggingState();
+
+	/** Adds a log entry for a given channel */
+	void AddStatisticLogEntry(const TPair<FTimedDataMonitorChannelIdentifier, FTimeDataChannelItem>& ChannelEntry);
+
 private:
 	bool bRequestSourceListRebuilt = false;
 	TMap<FTimedDataMonitorInputIdentifier, FTimeDataInputItem> InputMap;
 	TMap<FTimedDataMonitorChannelIdentifier, FTimeDataChannelItem> ChannelMap;
 	FSimpleMulticastDelegate OnIdentifierListChanged_Delegate;
+	float CachedTimecodeProviderFrameDelayInSeconds = 0.0f;
+
+#if !(UE_BUILD_SHIPPING || UE_BUILD_TEST)
+	struct FChannelStatisticLogging
+	{
+		FString ChannelName;
+		TArray<FString> Entries;
+	};
+
+	bool bHasStatFileLoggingStarted = false;
+	TMap<FTimedDataMonitorChannelIdentifier, FChannelStatisticLogging> StatLoggingMap;
+#endif //!(UE_BUILD_SHIPPING || UE_BUILD_TEST)
 };
