Fix deprecation warnings
#rnx

diff --git a/Engine/Plugins/Experimental/DigitalMultiplex/DMXEngine/Source/DMXEditor/Private/DMXEditor.cpp b/Engine/Plugins/Experimental/DigitalMultiplex/DMXEngine/Source/DMXEditor/Private/DMXEditor.cpp
--- a/Engine/Plugins/Experimental/DigitalMultiplex/DMXEngine/Source/DMXEditor/Private/DMXEditor.cpp
+++ b/Engine/Plugins/Experimental/DigitalMultiplex/DMXEngine/Source/DMXEditor/Private/DMXEditor.cpp
@@ -205,7 +205,7 @@ bool FDMXEditor::InvokeEditorTabFromEntityType(TSubclassOf<UDMXEntity> InEntityC
 		FName CurrentTab = FGlobalTabmanager::Get()->GetActiveTab()->GetLayoutIdentifier().TabType;
 		if (!CurrentTab.IsEqual(TargetTabId))
 		{
-			TabManager->InvokeTab(MoveTemp(TargetTabId));
+			TabManager->TryInvokeTab(MoveTemp(TargetTabId));
 		}
 		
 		return true;
diff --git a/Engine/Plugins/VirtualProduction/TimedDataMonitor/Source/TimedDataMonitorEditor/Private/STimedDataMonitorPanel.cpp b/Engine/Plugins/VirtualProduction/TimedDataMonitor/Source/TimedDataMonitorEditor/Private/STimedDataMonitorPanel.cpp
--- a/Engine/Plugins/VirtualProduction/TimedDataMonitor/Source/TimedDataMonitorEditor/Private/STimedDataMonitorPanel.cpp
+++ b/Engine/Plugins/VirtualProduction/TimedDataMonitor/Source/TimedDataMonitorEditor/Private/STimedDataMonitorPanel.cpp
@@ -655,7 +655,7 @@ bool STimedDataMonitorPanel::IsResetEvaluationChecked() const
 
 FReply STimedDataMonitorPanel::OnShowBuffersClicked()
 {
-	FGlobalTabmanager::Get()->InvokeTab(FTabId(Utilities::NAME_TimedDataMonitorBuffers));
+	FGlobalTabmanager::Get()->TryInvokeTab(FTabId(Utilities::NAME_TimedDataMonitorBuffers));
 	return FReply::Handled();
 }
 
diff --git a/Engine/Source/Developer/TraceInsights/Private/Insights/InsightsManager.cpp b/Engine/Source/Developer/TraceInsights/Private/Insights/InsightsManager.cpp
--- a/Engine/Source/Developer/TraceInsights/Private/Insights/InsightsManager.cpp
+++ b/Engine/Source/Developer/TraceInsights/Private/Insights/InsightsManager.cpp
@@ -219,7 +219,7 @@ void FInsightsManager::SpawnAndActivateTabs()
 	// Open Session Info tab.
 	if (FGlobalTabmanager::Get()->HasTabSpawner(FInsightsManagerTabs::SessionInfoTabId))
 	{
-		FGlobalTabmanager::Get()->InvokeTab(FInsightsManagerTabs::SessionInfoTabId);
+		FGlobalTabmanager::Get()->TryInvokeTab(FInsightsManagerTabs::SessionInfoTabId);
 	}
 
 	// Open Timing Insights tab.
diff --git a/Engine/Source/Editor/AnimationEditor/Private/AnimationEditor.cpp b/Engine/Source/Editor/AnimationEditor/Private/AnimationEditor.cpp
--- a/Engine/Source/Editor/AnimationEditor/Private/AnimationEditor.cpp
+++ b/Engine/Source/Editor/AnimationEditor/Private/AnimationEditor.cpp
@@ -451,7 +451,7 @@ TSharedPtr<SDockTab> FAnimationEditor::OpenNewAnimationDocumentTab(UAnimationAss
 		// Invoke the preview tab if this is a montage
 		if(InAnimAsset->IsA<UAnimMontage>())
 		{
-			TabManager->InvokeTab(AnimationEditorTabs::AnimMontageSectionsTab);
+			TabManager->TryInvokeTab(AnimationEditorTabs::AnimMontageSectionsTab);
 			OnSectionsChanged.Broadcast();
 		}
 		else
diff --git a/Engine/Source/Editor/LevelEditor/Private/LevelEditor.cpp b/Engine/Source/Editor/LevelEditor/Private/LevelEditor.cpp
--- a/Engine/Source/Editor/LevelEditor/Private/LevelEditor.cpp
+++ b/Engine/Source/Editor/LevelEditor/Private/LevelEditor.cpp
@@ -217,7 +217,7 @@ TSharedRef<SDockTab> FLevelEditorModule::SpawnLevelEditor( const FSpawnTabArgs&
 		{
 			// In legacy mode this toolbox should always be open
 			static const FTabId ToolboxTabId("LevelEditorToolBox");
-			LevelEditorTabManager->InvokeTab(ToolboxTabId);
+			LevelEditorTabManager->TryInvokeTab(ToolboxTabId);
 
 			// In legacy mode the standalone placement browser tab should not be opened
 			static const FTabId PlacementBrowserTabId("PlacementBrowser");
diff --git a/Engine/Source/Editor/LevelEditor/Private/SLevelEditor.cpp b/Engine/Source/Editor/LevelEditor/Private/SLevelEditor.cpp
--- a/Engine/Source/Editor/LevelEditor/Private/SLevelEditor.cpp
+++ b/Engine/Source/Editor/LevelEditor/Private/SLevelEditor.cpp
@@ -1460,7 +1460,7 @@ void SLevelEditor::OnEditorModeIdChanged(const FEditorModeID& ModeChangedID, boo
 		}
 		else
 		{
-			LevelEditorTabManager->InvokeTab(LevelEditorTabIds::LevelEditorToolBox);
+			LevelEditorTabManager->TryInvokeTab(LevelEditorTabIds::LevelEditorToolBox);
 		}
 	}
 }
diff --git a/Engine/Source/Editor/UnrealEd/Private/PlayLevelViaLauncher.cpp b/Engine/Source/Editor/UnrealEd/Private/PlayLevelViaLauncher.cpp
--- a/Engine/Source/Editor/UnrealEd/Private/PlayLevelViaLauncher.cpp
+++ b/Engine/Source/Editor/UnrealEd/Private/PlayLevelViaLauncher.cpp
@@ -22,7 +22,7 @@
 
 static void HandleHyperlinkNavigate()
 {
-	FGlobalTabmanager::Get()->InvokeTab(FName("OutputLog"));
+	FGlobalTabmanager::Get()->TryInvokeTab(FName("OutputLog"));
 }
 
 static void HandleCancelButtonClicked(ILauncherWorkerPtr LauncherWorker)
