Merging CL# 11619735

- Fixing default livelink preset not loading when depending on other module. Moved to PostEngineInit instead.

diff --git a/Engine/Plugins/Animation/LiveLink/Source/LiveLink/Private/LiveLinkModule.cpp b/Engine/Plugins/Animation/LiveLink/Source/LiveLink/Private/LiveLinkModule.cpp
--- a/Engine/Plugins/Animation/LiveLink/Source/LiveLink/Private/LiveLinkModule.cpp
+++ b/Engine/Plugins/Animation/LiveLink/Source/LiveLink/Private/LiveLinkModule.cpp
@@ -2,10 +2,11 @@
 
 #include "LiveLinkModule.h"
 
+#include "Interfaces/IPluginManager.h"
 #include "LiveLinkLogInstance.h"
 #include "LiveLinkPreset.h"
 #include "LiveLinkSettings.h"
-#include "Interfaces/IPluginManager.h"
+#include "Misc/CoreDelegates.h"
 
 #define LOCTEXT_NAMESPACE "LiveLinkModule"
 
@@ -29,14 +30,14 @@ void FLiveLinkModule::StartupModule()
 	IModularFeatures::Get().RegisterModularFeature(FLiveLinkClient::ModularFeatureName, &LiveLinkClient);
 	LiveLinkMotionController.RegisterController();
 
-	if (ULiveLinkPreset* Preset = GetDefault<ULiveLinkSettings>()->DefaultLiveLinkPreset.LoadSynchronous())
-	{
-		Preset->ApplyToClient();
-	}
+	//Register for engine initialization completed so we can load default preset if any. Presets could depend on plugins loaded at a later stage.
+	FCoreDelegates::OnFEngineLoopInitComplete.AddRaw(this, &FLiveLinkModule::OnEngineLoopInitComplete);
 }
 
 void FLiveLinkModule::ShutdownModule()
 {
+	FCoreDelegates::OnFEngineLoopInitComplete.RemoveAll(this);
+
 	HeartbeatEmitter->Exit();
 	DiscoveryManager->Stop();
 	LiveLinkMotionController.UnregisterController();
@@ -59,6 +60,14 @@ void FLiveLinkModule::CreateStyle()
 	StyleSet->Set("LiveLinkIcon", new FSlateImageBrush((ContentDir / TEXT("LiveLink_16x")) + TEXT(".png"), Icon16x16));
 }
 
+void FLiveLinkModule::OnEngineLoopInitComplete()
+{
+	if (ULiveLinkPreset* Preset = GetDefault<ULiveLinkSettings>()->DefaultLiveLinkPreset.LoadSynchronous())
+	{
+		Preset->ApplyToClient();
+	}
+}
+
 IMPLEMENT_MODULE(FLiveLinkModule, LiveLink);
 
 #undef LOCTEXT_NAMESPACE
\ No newline at end of file
diff --git a/Engine/Plugins/Animation/LiveLink/Source/LiveLink/Private/LiveLinkModule.h b/Engine/Plugins/Animation/LiveLink/Source/LiveLink/Private/LiveLinkModule.h
--- a/Engine/Plugins/Animation/LiveLink/Source/LiveLink/Private/LiveLinkModule.h
+++ b/Engine/Plugins/Animation/LiveLink/Source/LiveLink/Private/LiveLinkModule.h
@@ -5,11 +5,13 @@
 #include "ILiveLinkModule.h"
 
 #include "Features/IModularFeatures.h"
+
 #include "LiveLinkClient.h"
 #include "LiveLinkDebugCommand.h"
 #include "LiveLinkHeartbeatEmitter.h"
 #include "LiveLinkMessageBusDiscoveryManager.h"
 #include "LiveLinkMotionController.h"
+
 #include "Styling/SlateStyle.h"
 
 /**
@@ -37,6 +39,7 @@ public:
 
 private:
 	void CreateStyle();
+	void OnEngineLoopInitComplete();
 
 private:
 	friend FLiveLinkClientReference;
@@ -50,4 +53,4 @@ private:
 	TUniquePtr<FLiveLinkHeartbeatEmitter> HeartbeatEmitter;
 	TUniquePtr<FLiveLinkMessageBusDiscoveryManager> DiscoveryManager;
 	TUniquePtr<FLiveLinkDebugCommand> LiveLinkDebugCommand;
-};
\ No newline at end of file
+};
diff --git a/Engine/Plugins/VirtualProduction/TimedDataMonitor/Source/TimedDataMonitorEditor/Private/STimedDataMonitorPanel.cpp b/Engine/Plugins/VirtualProduction/TimedDataMonitor/Source/TimedDataMonitorEditor/Private/STimedDataMonitorPanel.cpp
--- a/Engine/Plugins/VirtualProduction/TimedDataMonitor/Source/TimedDataMonitorEditor/Private/STimedDataMonitorPanel.cpp
+++ b/Engine/Plugins/VirtualProduction/TimedDataMonitor/Source/TimedDataMonitorEditor/Private/STimedDataMonitorPanel.cpp
@@ -362,23 +362,45 @@ void STimedDataMonitorPanel::Construct(const FArguments& InArgs)
 				.FillWidth(0.5f)
 				.Padding(0.f, 0.f, 10.f, 0.f)
 				[
-					SNew(SBorder)
-					.HAlign(HAlign_Center)
-					.VAlign(VAlign_Center)
-					.Padding(FMargin(4.f, 4.f, 4.f, 4.f))
+					SNew(SVerticalBox)
+					+ SVerticalBox::Slot()
+					.AutoHeight()
+					[
+						SNew(STextBlock)
+						.Text(LOCTEXT("GenlockLabel", "Genlock"))
+					]
+					+ SVerticalBox::Slot()
+					.AutoHeight()
 					[
-						SAssignNew(TimedDataGenlockWidget, STimedDataGenlock, SharedThis<STimedDataMonitorPanel>(this))
+						SNew(SBorder)
+						.HAlign(HAlign_Center)
+						.VAlign(VAlign_Center)
+						.Padding(FMargin(4.f, 4.f, 4.f, 4.f))
+						[
+							SAssignNew(TimedDataGenlockWidget, STimedDataGenlock, SharedThis<STimedDataMonitorPanel>(this))
+						]
 					]
 				]
 				+ SHorizontalBox::Slot()
 				.FillWidth(0.5f)
 				.Padding(0.f)
 				[
-					SNew(SBorder)
-					.HAlign(HAlign_Center)
-					.VAlign(VAlign_Center)
+					SNew(SVerticalBox)
+					+ SVerticalBox::Slot()
+					.AutoHeight()
+					[
+						SNew(STextBlock)
+						.Text(LOCTEXT("TimecodeLabel", "Timecode Provider"))
+					]
+					+ SVerticalBox::Slot()
+					.AutoHeight()
 					[
-						SAssignNew(TimedDataTimecodeWidget, STimedDataTimecodeProvider, SharedThis<STimedDataMonitorPanel>(this))
+						SNew(SBorder)
+						.HAlign(HAlign_Center)
+						.VAlign(VAlign_Center)
+						[
+							SAssignNew(TimedDataTimecodeWidget, STimedDataTimecodeProvider, SharedThis<STimedDataMonitorPanel>(this))
+						]
 					]
 				]
 			]
