TimecodeDataMonitor: Fix property initialization.

diff --git a/Engine/Plugins/VirtualProduction/TimedDataMonitor/Source/TimedDataMonitor/Public/TimedDataMonitorCalibration.h b/Engine/Plugins/VirtualProduction/TimedDataMonitor/Source/TimedDataMonitor/Public/TimedDataMonitorCalibration.h
--- a/Engine/Plugins/VirtualProduction/TimedDataMonitor/Source/TimedDataMonitor/Public/TimedDataMonitorCalibration.h
+++ b/Engine/Plugins/VirtualProduction/TimedDataMonitor/Source/TimedDataMonitor/Public/TimedDataMonitorCalibration.h
@@ -45,7 +45,7 @@ struct FTimedDataMonitorCalibrationResult
 	GENERATED_BODY();
 
 	UPROPERTY(BlueprintReadOnly, VisibleAnywhere, Category = "Result")
-	ETimedDataMonitorCalibrationReturnCode ReturnCode;
+	ETimedDataMonitorCalibrationReturnCode ReturnCode = ETimedDataMonitorCalibrationReturnCode::Succeeded;
 
 	UPROPERTY(BlueprintReadOnly, VisibleAnywhere, Category = "Result")
 	TArray<FTimedDataMonitorInputIdentifier> FailureInputIdentifiers;
@@ -118,7 +118,7 @@ struct FTimedDataMonitorTimeCorrectionResult
 	GENERATED_BODY();
 
 	UPROPERTY(BlueprintReadOnly, VisibleAnywhere, Category = "Result")
-	ETimedDataMonitorTimeCorrectionReturnCode ReturnCode;
+	ETimedDataMonitorTimeCorrectionReturnCode ReturnCode = ETimedDataMonitorTimeCorrectionReturnCode::Succeeded;
 
 	UPROPERTY(BlueprintReadOnly, VisibleAnywhere, Category = "Result")
 	TArray<FTimedDataMonitorChannelIdentifier> FailureChannelIdentifiers;
diff --git a/Engine/Plugins/VirtualProduction/TimedDataMonitor/Source/TimedDataMonitor/Public/TimedDataMonitorSubsystem.h b/Engine/Plugins/VirtualProduction/TimedDataMonitor/Source/TimedDataMonitor/Public/TimedDataMonitorSubsystem.h
--- a/Engine/Plugins/VirtualProduction/TimedDataMonitor/Source/TimedDataMonitor/Public/TimedDataMonitorSubsystem.h
+++ b/Engine/Plugins/VirtualProduction/TimedDataMonitor/Source/TimedDataMonitor/Public/TimedDataMonitorSubsystem.h
@@ -50,7 +50,7 @@ public:
 	
 	/** New state of the channel */
 	UPROPERTY(VisibleAnywhere, Category = "Timed Data State")
-	ETimedDataInputState NewState;
+	ETimedDataInputState NewState = ETimedDataInputState::Connected;
 
 	/** Input owning that channel */
 	UPROPERTY(VisibleAnywhere, Category = "Timed Data State")
@@ -78,7 +78,7 @@ public:
 
 	/** New state of the channel */
 	UPROPERTY(VisibleAnywhere, Category = "Timed Data State")
-	ETimedDataMonitorEvaluationState NewState;
+	ETimedDataMonitorEvaluationState NewState = ETimedDataMonitorEvaluationState::NoSample;
 
 	/** Input owning that channel */
 	UPROPERTY(VisibleAnywhere, Category = "Timed Data State")
