- Fixing timed data monitor stage events type that made it unfilterable.
- Updated static assert to catch wrong base types for stage events

diff --git a/Engine/Plugins/VirtualProduction/TimedDataMonitor/Source/TimedDataMonitor/Public/TimedDataMonitorSubsystem.h b/Engine/Plugins/VirtualProduction/TimedDataMonitor/Source/TimedDataMonitor/Public/TimedDataMonitorSubsystem.h
--- a/Engine/Plugins/VirtualProduction/TimedDataMonitor/Source/TimedDataMonitor/Public/TimedDataMonitorSubsystem.h
+++ b/Engine/Plugins/VirtualProduction/TimedDataMonitor/Source/TimedDataMonitor/Public/TimedDataMonitorSubsystem.h
@@ -34,7 +34,7 @@ enum class ETimedDataMonitorEvaluationState : uint8
 };
 
 USTRUCT(meta = (DisplayName = "TimedSourcesConnectionEvent"))
-struct FTimedDataMonitorChannelConnectionStateEvent : public FStageProviderMessage
+struct FTimedDataMonitorChannelConnectionStateEvent : public FStageProviderEventMessage
 {
 	GENERATED_BODY()
 
@@ -62,7 +62,7 @@ public:
 };
 
 USTRUCT(meta = (DisplayName = "TimedSourcesEvaluationEvent"))
-struct FTimedDataMonitorChannelEvaluationStateEvent : public FStageProviderMessage
+struct FTimedDataMonitorChannelEvaluationStateEvent : public FStageProviderEventMessage
 {
 	GENERATED_BODY()
 
diff --git a/Engine/Source/Runtime/VirtualProduction/StageDataCore/Public/IStageDataProvider.h b/Engine/Source/Runtime/VirtualProduction/StageDataCore/Public/IStageDataProvider.h
--- a/Engine/Source/Runtime/VirtualProduction/StageDataCore/Public/IStageDataProvider.h
+++ b/Engine/Source/Runtime/VirtualProduction/StageDataCore/Public/IStageDataProvider.h
@@ -73,7 +73,7 @@ private:
 		{
 			IStageDataProvider* StageDataProvider = &IModularFeatures::Get().GetModularFeature<IStageDataProvider>(IStageDataProvider::ModularFeatureName);
 
-			static_assert(TIsDerivedFrom<MessageType, FStageProviderMessage>::IsDerived, "MessageType must be a FStageProviderMessage derived UStruct.");
+			static_assert(TIsDerivedFrom<MessageType, FStageProviderEventMessage>::IsDerived || TIsDerivedFrom<MessageType, FStageProviderPeriodicMessage>::IsDerived, "MessageType must be a FStageProviderEventMessage or FStageProviderPeriodicMessage derived UStruct.");
 			return StageDataProvider->SendMessageInternal(&Message, MessageType::StaticStruct(), InFlags);
 		}
 
