- Updating StageMonitor and TimedDataMonitor with official icons

diff --git a/Engine/Plugins/VirtualProduction/StageMonitoring/Source/StageMonitorEditor/Private/SStageMonitorPanel.cpp b/Engine/Plugins/VirtualProduction/StageMonitoring/Source/StageMonitorEditor/Private/SStageMonitorPanel.cpp
--- a/Engine/Plugins/VirtualProduction/StageMonitoring/Source/StageMonitorEditor/Private/SStageMonitorPanel.cpp
+++ b/Engine/Plugins/VirtualProduction/StageMonitoring/Source/StageMonitorEditor/Private/SStageMonitorPanel.cpp
@@ -62,7 +62,7 @@ void SStageMonitorPanel::RegisterNomadTabSpawner(TSharedRef<FWorkspaceItem> InWo
 			.SetDisplayName(LOCTEXT("TabTitle", "Stage Monitor"))
 			.SetTooltipText(LOCTEXT("TooltipText", "Monitor performance data from stage machines"))
 			.SetGroup(InWorkspaceItem)
-			.SetIcon(FSlateIcon(FEditorStyle::GetStyleSetName(), "UnrealInsights.MenuIcon"));
+			.SetIcon(FSlateIcon(FStageMonitorEditorStyle::Get().GetStyleSetName(), "StageMonitor.TabIcon"));
 	};
 
 	FLevelEditorModule& LevelEditorModule = FModuleManager::LoadModuleChecked<FLevelEditorModule>(StageMonitorUtilities::NAME_LevelEditorModuleName);
diff --git a/Engine/Plugins/VirtualProduction/StageMonitoring/Source/StageMonitorEditor/Private/StageMonitorEditorStyle.cpp b/Engine/Plugins/VirtualProduction/StageMonitoring/Source/StageMonitorEditor/Private/StageMonitorEditorStyle.cpp
--- a/Engine/Plugins/VirtualProduction/StageMonitoring/Source/StageMonitorEditor/Private/StageMonitorEditorStyle.cpp
+++ b/Engine/Plugins/VirtualProduction/StageMonitoring/Source/StageMonitorEditor/Private/StageMonitorEditorStyle.cpp
@@ -4,8 +4,9 @@
 
 #include "EditorStyleSet.h"
 #include "Interfaces/IPluginManager.h"
-#include "Styling/SlateTypes.h"
 #include "Styling/CoreStyle.h"
+#include "Styling/SlateStyleMacros.h"
+#include "Styling/SlateTypes.h"
 
 
 
@@ -16,6 +17,12 @@ FStageMonitorEditorStyle::FStageMonitorEditorStyle()
 	SetContentRoot(ContentDir);
 	SetCoreContentRoot(FPaths::EngineContentDir() / TEXT("Slate"));
 
+	// Icons
+	{
+		const FVector2D Icon20x20(20.0f, 20.0f);
+		Set("StageMonitor.TabIcon", new IMAGE_BRUSH_SVG("StageMonitor", Icon20x20));
+	}
+
 	// ListView
 	{
 		const FVector2D Icon8x8(8.0f, 8.0f);
@@ -59,4 +66,3 @@ FStageMonitorEditorStyle& FStageMonitorEditorStyle::Get()
 	return Inst;
 }
 
-
diff --git a/Engine/Plugins/VirtualProduction/TimedDataMonitor/Source/TimedDataMonitorEditor/Private/STimedDataMonitorPanel.cpp b/Engine/Plugins/VirtualProduction/TimedDataMonitor/Source/TimedDataMonitorEditor/Private/STimedDataMonitorPanel.cpp
--- a/Engine/Plugins/VirtualProduction/TimedDataMonitor/Source/TimedDataMonitorEditor/Private/STimedDataMonitorPanel.cpp
+++ b/Engine/Plugins/VirtualProduction/TimedDataMonitor/Source/TimedDataMonitorEditor/Private/STimedDataMonitorPanel.cpp
@@ -82,7 +82,7 @@ void STimedDataMonitorPanel::RegisterNomadTabSpawner(TSharedRef<FWorkspaceItem>
 			.SetDisplayName(LOCTEXT("TabTitle", "Timed Data Monitor"))
 			.SetTooltipText(LOCTEXT("TooltipText", "Monitor inputs that can be time synchronized"))
 			.SetGroup(InWorkspaceItem)
-			.SetIcon(FSlateIcon(FTimedDataMonitorEditorStyle::Get().GetStyleSetName(), "Img.TimedDataMonitor.Small"));
+			.SetIcon(FSlateIcon(FTimedDataMonitorEditorStyle::Get().GetStyleSetName(), "Img.TimedDataMonitor.Icon"));
 
 		FGlobalTabmanager::Get()->RegisterNomadTabSpawner(Utilities::NAME_TimedDataMonitorBuffers, FOnSpawnTab::CreateStatic(&Utilities::CreateBuffersVisualizer))
 			.SetDisplayName(LOCTEXT("BufferVisualizerTitle", "Buffer Visualizer"))
diff --git a/Engine/Plugins/VirtualProduction/TimedDataMonitor/Source/TimedDataMonitorEditor/Private/TimedDataMonitorEditorStyle.cpp b/Engine/Plugins/VirtualProduction/TimedDataMonitor/Source/TimedDataMonitorEditor/Private/TimedDataMonitorEditorStyle.cpp
--- a/Engine/Plugins/VirtualProduction/TimedDataMonitor/Source/TimedDataMonitorEditor/Private/TimedDataMonitorEditorStyle.cpp
+++ b/Engine/Plugins/VirtualProduction/TimedDataMonitor/Source/TimedDataMonitorEditor/Private/TimedDataMonitorEditorStyle.cpp
@@ -4,8 +4,9 @@
 
 #include "EditorStyleSet.h"
 #include "Interfaces/IPluginManager.h"
-#include "Styling/SlateTypes.h"
 #include "Styling/CoreStyle.h"
+#include "Styling/SlateStyleMacros.h"
+#include "Styling/SlateTypes.h"
 
 const FName FTimedDataMonitorEditorStyle::NAME_TimecodeBrush = "Img.Timecode.Small";
 const FName FTimedDataMonitorEditorStyle::NAME_PlatformTimeBrush = "Img.PlatformTime.Small";
@@ -50,6 +51,9 @@ FTimedDataMonitorEditorStyle::FTimedDataMonitorEditorStyle()
 
 	// images
 	{
+		const FVector2D Icon20x20(20.0f, 20.0f);
+		Set("Img.TimedDataMonitor.Icon", new IMAGE_BRUSH_SVG("Common/TimedDataMonitor", Icon20x20));
+
 		Set(NAME_TimecodeBrush, new FSlateImageBrush(RootToContentDir(TEXT("Widgets/Timecode_16x.png")), Icon16x16));
 		Set(NAME_PlatformTimeBrush, new FSlateImageBrush(RootToContentDir(TEXT("Widgets/Time_16x.png")), Icon16x16));
 		Set(NAME_NoEvaluationBrush, new FSlateImageBrush(RootToContentDir(TEXT("Widgets/NoEvaluation_16x.png")), Icon16x16));
