- Removing unused icon from timed data monitor plugin

diff --git a/Engine/Plugins/VirtualProduction/TimedDataMonitor/Source/TimedDataMonitorEditor/Private/TimedDataMonitorEditorStyle.cpp b/Engine/Plugins/VirtualProduction/TimedDataMonitor/Source/TimedDataMonitorEditor/Private/TimedDataMonitorEditorStyle.cpp
--- a/Engine/Plugins/VirtualProduction/TimedDataMonitor/Source/TimedDataMonitorEditor/Private/TimedDataMonitorEditorStyle.cpp
+++ b/Engine/Plugins/VirtualProduction/TimedDataMonitor/Source/TimedDataMonitorEditor/Private/TimedDataMonitorEditorStyle.cpp
@@ -57,7 +57,6 @@ FTimedDataMonitorEditorStyle::FTimedDataMonitorEditorStyle()
 		Set(NAME_TimecodeBrush, new FSlateImageBrush(RootToContentDir(TEXT("Widgets/Timecode_16x.png")), Icon16x16));
 		Set(NAME_PlatformTimeBrush, new FSlateImageBrush(RootToContentDir(TEXT("Widgets/Time_16x.png")), Icon16x16));
 		Set(NAME_NoEvaluationBrush, new FSlateImageBrush(RootToContentDir(TEXT("Widgets/NoEvaluation_16x.png")), Icon16x16));
-		Set("Img.TimedDataMonitor.Small", new FSlateImageBrush(RootToContentDir(TEXT("Widgets/TimedDataMonitor_16x.png")), Icon16x16));
 
 		Set("Img.BufferVisualization", new FSlateImageBrush(RootToContentDir(TEXT("Widgets/BufferVisualization_24x.png")), Icon24x24));
 		Set("Img.Calibration", new FSlateImageBrush(RootToContentDir(TEXT("Widgets/Calibration_24x.png")), Icon24x24));
