[Timed Data Monitor] Draw vertical bars in buffer visualization at the world time of each frame in each subject's frame buffer

diff --git a/Engine/Plugins/VirtualProduction/TimedDataMonitor/Source/TimedDataMonitor/Private/TimedDataMonitorSubsystem.cpp b/Engine/Plugins/VirtualProduction/TimedDataMonitor/Source/TimedDataMonitor/Private/TimedDataMonitorSubsystem.cpp
--- a/Engine/Plugins/VirtualProduction/TimedDataMonitor/Source/TimedDataMonitor/Private/TimedDataMonitorSubsystem.cpp
+++ b/Engine/Plugins/VirtualProduction/TimedDataMonitor/Source/TimedDataMonitor/Private/TimedDataMonitorSubsystem.cpp
@@ -785,6 +785,21 @@ FTimedDataChannelSampleTime UTimedDataMonitorSubsystem::GetChannelNewestDataTime
 }
 
 
+TArray<FTimedDataChannelSampleTime> UTimedDataMonitorSubsystem::GetChannelFrameDataTimes(const FTimedDataMonitorChannelIdentifier& Identifier)
+{
+	BuildSourcesListIfNeeded();
+
+	TArray<FTimedDataChannelSampleTime> SampleTimes;
+
+	if (const FTimeDataChannelItem* SourceItem = ChannelMap.Find(Identifier))
+	{
+		SampleTimes = SourceItem->Channel->GetDataTimes();
+	}
+
+	return SampleTimes;
+}
+
+
 int32 UTimedDataMonitorSubsystem::GetChannelNumberOfSamples(const FTimedDataMonitorChannelIdentifier& Identifier)
 {
 	BuildSourcesListIfNeeded();
diff --git a/Engine/Plugins/VirtualProduction/TimedDataMonitor/Source/TimedDataMonitor/Public/TimedDataMonitorSubsystem.h b/Engine/Plugins/VirtualProduction/TimedDataMonitor/Source/TimedDataMonitor/Public/TimedDataMonitorSubsystem.h
--- a/Engine/Plugins/VirtualProduction/TimedDataMonitor/Source/TimedDataMonitor/Public/TimedDataMonitorSubsystem.h
+++ b/Engine/Plugins/VirtualProduction/TimedDataMonitor/Source/TimedDataMonitor/Public/TimedDataMonitorSubsystem.h
@@ -370,6 +370,10 @@ public:
 	UFUNCTION(BlueprintCallable, Category = "Timed Data Monitor|Channel")
 	FTimedDataChannelSampleTime GetChannelNewestDataTime(const FTimedDataMonitorChannelIdentifier& Identifier);
 
+	/** Get the sample times for every frame in the channel */
+	UFUNCTION(BlueprintCallable, Category = "Timed Data Monitor|Channel")
+	TArray<FTimedDataChannelSampleTime> GetChannelFrameDataTimes(const FTimedDataMonitorChannelIdentifier& Identifier);
+
 	/** Get the number of data samples available. */
 	UFUNCTION(BlueprintCallable, Category = "Timed Data Monitor|Channel")
 	int32 GetChannelNumberOfSamples(const FTimedDataMonitorChannelIdentifier& Identifier);
diff --git a/Engine/Plugins/VirtualProduction/TimedDataMonitor/Source/TimedDataMonitorEditor/Private/STimingDiagramWidget.cpp b/Engine/Plugins/VirtualProduction/TimedDataMonitor/Source/TimedDataMonitorEditor/Private/STimingDiagramWidget.cpp
--- a/Engine/Plugins/VirtualProduction/TimedDataMonitor/Source/TimedDataMonitorEditor/Private/STimingDiagramWidget.cpp
+++ b/Engine/Plugins/VirtualProduction/TimedDataMonitor/Source/TimedDataMonitorEditor/Private/STimingDiagramWidget.cpp
@@ -83,6 +83,9 @@ public:
 		const float MeanLocationMinX = LocationOfCenter - (OldestMean * SizePerSeconds);
 		const float MeanLocationMaxX = LocationOfCenter + (NewestMean * SizePerSeconds);
 
+		const bool bDrawFrameTimes = GetDefault<UTimedDataMonitorEditorSettings>()->bDrawFrameTimesInBufferVisualization;
+		const bool bUseAccurateFrameTimes = GetDefault<UTimedDataMonitorEditorSettings>()->bUseAccurateFrameTimesInBufferVisualization;
+
 		// square to show that the data goes further
 		if (bShowFurther)
 		{
@@ -200,6 +203,46 @@ public:
 				DarkBrush->GetTint(InWidgetStyle) * EvaluationColor);
 		}
 
+		// Draw frame lines
+		if (bDrawFrameTimes && (SampleCount > 1))
+		{
+			const FLinearColor EvaluationColor = FLinearColor::Black;
+			const double FrameIntervalEstimate = (MaxSampleTime - MinSampleTime) / (SampleCount - 1);
+			for (int32 FrameIndex = 0; FrameIndex < SampleCount; ++FrameIndex)
+			{
+				double SampleTime = 0.0;
+				if (bUseAccurateFrameTimes && (AllSampleTimes.Num() == SampleCount))
+				{
+					SampleTime = AllSampleTimes[FrameIndex];
+				}
+				else
+				{
+					SampleTime = MinSampleTime + (FrameIntervalEstimate * FrameIndex);
+				}
+
+				const double LocationToDraw = LocationOfCenter + ((SampleTime - EvaluationTime) * SizePerSeconds);
+
+				TArray<FVector2D> LinePoints;
+				LinePoints.SetNum(2);
+
+				if (LocationToDraw >= SizeOfFurthur)
+				{
+					LinePoints[0] = FVector2D(LocationToDraw, 0.f);
+					LinePoints[1] = FVector2D(LocationToDraw, SizeY * 0.5f);
+
+					FSlateDrawElement::MakeLines(
+						OutDrawElements,
+						LayerId,
+						AllottedGeometry.ToPaintGeometry(),
+						LinePoints,
+						ESlateDrawEffect::None,
+						DarkBrush->GetTint(InWidgetStyle) * EvaluationColor,
+						false
+					);
+				}
+			}
+		}
+
 		return LayerId;
 	}
 
@@ -220,11 +263,13 @@ public:
 	double EvaluationTime = 0.0;
 	double MinSampleTime = 0.0;
 	double MaxSampleTime = 0.0;
+	TArray<double> AllSampleTimes;
 	double OldestMean = 0.0;
 	double NewestMean = 0.0;
 	double OldestSigma = 0.0;
 	double NewestSigma = 0.0;
 	int32 NumberOfSigma = 3;
+	int32 SampleCount = 0;
 	FSlateFontInfo FontInfo;
 };
 
@@ -276,10 +321,14 @@ void STimingDiagramWidget::UpdateCachedValue()
 
 	const ETimedDataInputEvaluationType EvaluationType = TimedDataMonitorSubsystem->GetInputEvaluationType(InputIdentifier);
 
+	const UTimedDataMonitorEditorSettings* const EditorSettings = GetDefault<UTimedDataMonitorEditorSettings>();
+	const bool bUpdateSampleTimes = EditorSettings->bDrawFrameTimesInBufferVisualization && EditorSettings->bUseAccurateFrameTimesInBufferVisualization;
+
 	double EvaluationOffset = TimedDataMonitorSubsystem->GetInputEvaluationOffsetInSeconds(InputIdentifier);
 	double MinSampleTime = 0.0;
 	double MaxSampleTime = 0.0;
 	double NewestMean = 0.0;
+
 	if (bIsInput)
 	{
 		MinSampleTime = TimedDataMonitorSubsystem->GetInputOldestDataTime(InputIdentifier).AsSeconds(EvaluationType);
@@ -297,10 +346,15 @@ void STimingDiagramWidget::UpdateCachedValue()
 		GraphicWidget->NewestMean = TimedDataMonitorSubsystem->GetChannelEvaluationDistanceToNewestSampleMean(ChannelIdentifier);
 		GraphicWidget->OldestSigma = TimedDataMonitorSubsystem->GetChannelEvaluationDistanceToOldestSampleStandardDeviation(ChannelIdentifier);
 		GraphicWidget->NewestSigma = TimedDataMonitorSubsystem->GetChannelEvaluationDistanceToNewestSampleStandardDeviation(ChannelIdentifier);
+		GraphicWidget->SampleCount = TimedDataMonitorSubsystem->GetChannelNumberOfSamples(ChannelIdentifier);
+
+		if (bUpdateSampleTimes)
+		{
+			UpdateSampleTimes(TimedDataMonitorSubsystem->GetChannelFrameDataTimes(ChannelIdentifier), EvaluationType);
+		}
 	}
 	GraphicWidget->MinSampleTime = MinSampleTime + EvaluationOffset;
 	GraphicWidget->MaxSampleTime = MaxSampleTime + EvaluationOffset;
-
 	GraphicWidget->EvaluationTime = UTimedDataMonitorSubsystem::GetEvaluationTime(EvaluationType);
 }
 
@@ -342,8 +396,50 @@ FText STimingDiagramWidget::GetTooltipText() const
 		, DistanceToOldestSampleAverage
 		, DistanceToOldestSampleSigma);
 	}
+}
 
-	
+void STimingDiagramWidget::UpdateSampleTimes(const TArray<FTimedDataChannelSampleTime>& FrameDataTimes, ETimedDataInputEvaluationType EvaluationType)
+{
+	// Cache the newest time in the buffer from the last update
+	double PreviousNewestTime = -1.0;
+	if (GraphicWidget->AllSampleTimes.Num() > 0)
+	{
+		PreviousNewestTime = GraphicWidget->AllSampleTimes.Last();
+	}
+
+	// If the widget's buffer is too large (because the buffer size changed), resize it and remove the oldest frames
+	if (GraphicWidget->AllSampleTimes.Num() > FrameDataTimes.Num())
+	{
+		const int32 Count = (GraphicWidget->AllSampleTimes.Num() - FrameDataTimes.Num());
+		constexpr bool bAllowShrinking = true;
+		GraphicWidget->AllSampleTimes.RemoveAt(0, Count, bAllowShrinking);
+	}
+	else if (GraphicWidget->AllSampleTimes.Num() < FrameDataTimes.Num())
+	{
+		GraphicWidget->AllSampleTimes.Reserve(FrameDataTimes.Num());
+	}
+
+	// Add all sample times that are newer than the previous newest time
+	for (int32 SampleIndex = FrameDataTimes.Num() - 1; SampleIndex >= 0; --SampleIndex)
+	{
+		const double SampleTimeInSeconds = FrameDataTimes[SampleIndex].AsSeconds(EvaluationType);
+		if (SampleTimeInSeconds > PreviousNewestTime)
+		{
+			// Remove the oldest sample(s) from the widget's buffer if it full 
+			if (GraphicWidget->AllSampleTimes.Num() >= FrameDataTimes.Num())
+			{
+				constexpr int32 Count = 1;
+				constexpr bool bAllowShrinking = false;
+				GraphicWidget->AllSampleTimes.RemoveAt(0, Count, bAllowShrinking);
+			}
+
+			GraphicWidget->AllSampleTimes.Add(SampleTimeInSeconds);
+		}
+		else
+		{
+			break;
+		}
+	}
 
 }
 
diff --git a/Engine/Plugins/VirtualProduction/TimedDataMonitor/Source/TimedDataMonitorEditor/Private/STimingDiagramWidget.h b/Engine/Plugins/VirtualProduction/TimedDataMonitor/Source/TimedDataMonitorEditor/Private/STimingDiagramWidget.h
--- a/Engine/Plugins/VirtualProduction/TimedDataMonitor/Source/TimedDataMonitorEditor/Private/STimingDiagramWidget.h
+++ b/Engine/Plugins/VirtualProduction/TimedDataMonitor/Source/TimedDataMonitorEditor/Private/STimingDiagramWidget.h
@@ -43,7 +43,8 @@ public:
 
 private:
 	FText GetTooltipText() const;
-	
+	void UpdateSampleTimes(const TArray<FTimedDataChannelSampleTime>& FrameDataTimes, ETimedDataInputEvaluationType EvaluationType);
+
 private:
 	// Identifier of the channel this widget is associated to
 	TSharedPtr<STimingDiagramWidgetGraphic> GraphicWidget;
diff --git a/Engine/Plugins/VirtualProduction/TimedDataMonitor/Source/TimedDataMonitorEditor/Private/TimedDataMonitorEditorSettings.h b/Engine/Plugins/VirtualProduction/TimedDataMonitor/Source/TimedDataMonitorEditor/Private/TimedDataMonitorEditorSettings.h
--- a/Engine/Plugins/VirtualProduction/TimedDataMonitor/Source/TimedDataMonitorEditor/Private/TimedDataMonitorEditorSettings.h
+++ b/Engine/Plugins/VirtualProduction/TimedDataMonitor/Source/TimedDataMonitorEditor/Private/TimedDataMonitorEditorSettings.h
@@ -38,6 +38,14 @@ public:
 	UPROPERTY(Config, EditAnywhere, Category ="UI", meta=(ClampMin=0.0f))
 	float RefreshRate = 0.2f;
 
+	/** Whether or not to draw vertical bars delineating each frame in the buffer of a channel */
+	UPROPERTY(Config, EditAnywhere, Category = "UI")
+	bool bDrawFrameTimesInBufferVisualization = true;
+
+	/** If true, accurate world times will be used to draw frame times. If false, a constant frame interval is used */
+	UPROPERTY(Config, EditAnywhere, Category = "UI", meta = (EditCondition = "bDrawFrameTimesInBufferVisualization"))
+	bool bUseAccurateFrameTimesInBufferVisualization = false;
+
 	UPROPERTY(Config, AdvancedDisplay, EditAnywhere, Category = "UI", meta = (InlineEditConditionToggle = true))
 	bool bOverrideNumberOfStandardDeviationToShow;
 
