TDM Improvements -
Put channels in alphabetical order in buffer visualizer
Fix uninitialized argument for buffer visualizer
Add buttons to increment/decrement frame offset and time evaluation offset
Add dropdown to select evaluation offset type (frame vs seconds)
Fix style not being registered resulting in white texture
Changed progress bar to custom frame time widget
Fix  borders in timed data monitor panel

diff --git a/Engine/Plugins/Animation/LiveLink/Source/LiveLink/Private/LiveLinkModule.cpp b/Engine/Plugins/Animation/LiveLink/Source/LiveLink/Private/LiveLinkModule.cpp
--- a/Engine/Plugins/Animation/LiveLink/Source/LiveLink/Private/LiveLinkModule.cpp
+++ b/Engine/Plugins/Animation/LiveLink/Source/LiveLink/Private/LiveLinkModule.cpp
@@ -8,6 +8,7 @@
 #include "LiveLinkSettings.h"
 #include "Misc/CommandLine.h"
 #include "Misc/CoreDelegates.h"
+#include "Styling/SlateStyleRegistry.h"
 
 #define LOCTEXT_NAMESPACE "LiveLinkModule"
 
@@ -46,6 +47,7 @@ void FLiveLinkModule::ShutdownModule()
 	IModularFeatures::Get().UnregisterModularFeature(FLiveLinkClient::ModularFeatureName, &LiveLinkClient);
 	FPlatformAtomics::InterlockedExchangePtr((void**)&LiveLinkClient_AnyThread, nullptr);
 
+	FSlateStyleRegistry::UnRegisterSlateStyle(*StyleSet.Get());
 	FLiveLinkLogInstance::DestroyInstance();
 }
 
@@ -53,6 +55,7 @@ void FLiveLinkModule::CreateStyle()
 {
 	static FName LiveLinkStyle(TEXT("LiveLinkCoreStyle"));
 	StyleSet = MakeShared<FSlateStyleSet>(LiveLinkStyle);
+	FSlateStyleRegistry::RegisterSlateStyle(*StyleSet.Get());
 
 	FString ContentDir = IPluginManager::Get().FindPlugin(TEXT("LiveLink"))->GetContentDir();
 
diff --git a/Engine/Plugins/VirtualProduction/TimedDataMonitor/Source/TimedDataMonitor/Private/TimedDataMonitorSubsystem.cpp b/Engine/Plugins/VirtualProduction/TimedDataMonitor/Source/TimedDataMonitor/Private/TimedDataMonitorSubsystem.cpp
--- a/Engine/Plugins/VirtualProduction/TimedDataMonitor/Source/TimedDataMonitor/Private/TimedDataMonitorSubsystem.cpp
+++ b/Engine/Plugins/VirtualProduction/TimedDataMonitor/Source/TimedDataMonitor/Private/TimedDataMonitorSubsystem.cpp
@@ -480,6 +480,29 @@ void UTimedDataMonitorSubsystem::SetInputEvaluationOffsetInSeconds(const FTimedD
 	}
 }
 
+float UTimedDataMonitorSubsystem::GetInputEvaluationOffsetInFrames(const FTimedDataMonitorInputIdentifier& Identifier)
+{
+	BuildSourcesListIfNeeded();
+
+	if (const FTimeDataInputItem* SourceItem = InputMap.Find(Identifier))
+	{
+		return SourceItem->Input->ConvertSecondOffsetInFrameOffset(SourceItem->Input->GetEvaluationOffsetInSeconds());
+	}
+
+	return 0.f;
+}
+
+
+void UTimedDataMonitorSubsystem::SetInputEvaluationOffsetInFrames(const FTimedDataMonitorInputIdentifier& Identifier, float Offset)
+{
+	BuildSourcesListIfNeeded();
+
+	if (FTimeDataInputItem* SourceItem = InputMap.Find(Identifier))
+	{
+		return SourceItem->Input->SetEvaluationOffsetInSeconds(SourceItem->Input->ConvertFrameOffsetInSecondOffset(Offset));
+	}
+}
+
 
 FFrameRate UTimedDataMonitorSubsystem::GetInputFrameRate(const FTimedDataMonitorInputIdentifier& Identifier)
 {
diff --git a/Engine/Plugins/VirtualProduction/TimedDataMonitor/Source/TimedDataMonitor/Public/TimedDataMonitorSubsystem.h b/Engine/Plugins/VirtualProduction/TimedDataMonitor/Source/TimedDataMonitor/Public/TimedDataMonitorSubsystem.h
--- a/Engine/Plugins/VirtualProduction/TimedDataMonitor/Source/TimedDataMonitor/Public/TimedDataMonitorSubsystem.h
+++ b/Engine/Plugins/VirtualProduction/TimedDataMonitor/Source/TimedDataMonitor/Public/TimedDataMonitorSubsystem.h
@@ -283,6 +283,14 @@ public:
 	UFUNCTION(BlueprintCallable, Category = "Timed Data Monitor|Input")
 	void SetInputEvaluationOffsetInSeconds(const FTimedDataMonitorInputIdentifier& Identifier, float Seconds);
 
+	/** Get the offset in frames (see GetEvaluationType) used at evaluation. */
+	UFUNCTION(BlueprintCallable, Category = "Timed Data Monitor|Input")
+	float GetInputEvaluationOffsetInFrames(const FTimedDataMonitorInputIdentifier& Identifier);
+
+	/** Set the offset in frames (see GetEvaluationType) used at evaluation. */
+	UFUNCTION(BlueprintCallable, Category = "Timed Data Monitor|Input")
+	void SetInputEvaluationOffsetInFrames(const FTimedDataMonitorInputIdentifier& Identifier, float Frames);
+
 	/** Get the frame rate at which the samples is produce. */
 	UFUNCTION(BlueprintCallable, Category = "Timed Data Monitor|Input")
 	FFrameRate GetInputFrameRate(const FTimedDataMonitorInputIdentifier& Identifier);
diff --git a/Engine/Plugins/VirtualProduction/TimedDataMonitor/Source/TimedDataMonitorEditor/Private/SFrameTime.h b/Engine/Plugins/VirtualProduction/TimedDataMonitor/Source/TimedDataMonitorEditor/Private/SFrameTime.h
new file mode 100644
--- /dev/null
+++ b/Engine/Plugins/VirtualProduction/TimedDataMonitor/Source/TimedDataMonitorEditor/Private/SFrameTime.h
@@ -0,0 +1,204 @@
+// Copyright Epic Games, Inc. All Rights Reserved.
+
+#pragma once
+
+#include "CoreMinimal.h"
+#include "Rendering/DrawElements.h"
+#include "Styling/AppStyle.h"
+#include "TimedDataMonitorEditorSettings.h"
+#include "Widgets/DeclarativeSyntaxSupport.h"
+#include "Widgets/SLeafWidget.h"
+
+class SFrameTime : public SLeafWidget
+{
+public:
+	SLATE_BEGIN_ARGS(SFrameTime)
+			: _FillImage(FAppStyle::GetBrush(TEXT("WhiteBrush")))
+			, _BackgroundImage(&FCoreStyle::Get().GetWidgetStyle<FProgressBarStyle>("ProgressBar").BackgroundImage)
+		{
+		}
+
+		SLATE_ARGUMENT(TOptional<double>, RefreshRate)
+		SLATE_ARGUMENT(const FSlateBrush*, FillImage)
+		SLATE_ARGUMENT(const FSlateBrush*, BackgroundImage)
+	SLATE_END_ARGS()
+
+	void Construct(const FArguments& InArgs)
+	{
+		RefreshRate = InArgs._RefreshRate;
+		FillImage = InArgs._FillImage;
+		BackgroundImage = InArgs._BackgroundImage;
+	}
+	
+	virtual void Tick(const FGeometry& AllottedGeometry, const double InCurrentTime, const float InDeltaTime) override
+	{
+		if (InCurrentTime - LastUpdateTime > RefreshRate.Get(GetDefault<UTimedDataMonitorEditorSettings>()->RefreshRate))
+		{
+			LastUpdateTime = InCurrentTime;
+
+			const float MaxTickRate = GEngine->GetMaxTickRate(0.001f, false);
+			const float TargetFPS = GEngine->bUseFixedFrameRate ? GEngine->FixedFrameRate : MaxTickRate;
+
+			if (FApp::GetIdleTime() == 0 && GEngine->bUseFixedFrameRate && TargetFPS != 0)
+			{
+				CachedFPSFraction = InDeltaTime / (1 / TargetFPS);
+			}
+			else
+			{
+				CachedFPSFraction = FApp::GetDeltaTime() > 0.0 ? (1.f - float(FApp::GetIdleTime() / FApp::GetDeltaTime())) : 1.f;
+			}
+		}
+
+	}
+
+	virtual int32 OnPaint(const FPaintArgs& Args, const FGeometry& AllottedGeometry, const FSlateRect& MyCullingRect, FSlateWindowElementList& OutDrawElements, int32 LayerId, const FWidgetStyle& InWidgetStyle, bool bParentEnabled) const override
+	{
+		// Used to track the layer ID we will return.
+		int32 RetLayerId = LayerId;
+		bool bEnabled = ShouldBeEnabled(bParentEnabled);
+		const ESlateDrawEffect DrawEffects = bEnabled ? ESlateDrawEffect::None : ESlateDrawEffect::DisabledEffect;
+
+		FSlateDrawElement::MakeBox(
+			OutDrawElements,
+			RetLayerId++,
+			AllottedGeometry.ToPaintGeometry(),
+			BackgroundImage,
+			DrawEffects,
+			InWidgetStyle.GetColorAndOpacityTint() * BackgroundImage->GetTint(InWidgetStyle)
+		);
+
+		const double WarnThreshold = GetDefault<UTimedDataMonitorEditorSettings>()->FrameTimeWarnThreshold;
+
+		FSlateColor FillColorAndOpacity = FLinearColor::Green;
+		if (CachedFPSFraction >= WarnThreshold && CachedFPSFraction < 1.0)
+		{
+			FillColorAndOpacity = FLinearColor(FColor::Orange);
+		}
+		else if (CachedFPSFraction >= 1.0)
+		{
+			FillColorAndOpacity = FLinearColor(FColor::Red);
+		}
+
+		const FLinearColor FillColorAndOpacitySRGB(InWidgetStyle.GetColorAndOpacityTint() * FillColorAndOpacity.GetColor(InWidgetStyle) * FillImage->GetTint(InWidgetStyle));
+
+		float MarqueeAnimOffsetX = 0.0f;
+		float MarqueeImageSizeX = 0.0f;
+		float MarqueeAnimOffsetY = 0.0f;
+		float MarqueeImageSizeY = 0.0f;
+
+		bool bScaleWithFillPerc = true;
+
+		const float XLimit = GetDefault<UTimedDataMonitorEditorSettings>()->TargetFrameTimePercent;
+		const float ClampedFraction = FMath::Clamp(CachedFPSFraction * XLimit, 0.0f, 1.0f);
+
+		const FVector2D BorderPadding = FVector2D{ 0.0, 0.0 };
+	
+		if (PushTransformedClip(OutDrawElements, AllottedGeometry, BorderPadding, FVector2D(0, 0), FSlateRect(0, 0, ClampedFraction, 1)))
+		{
+			FPaintGeometry PaintRect;
+			PaintRect = AllottedGeometry.ToPaintGeometry(
+				FVector2D(0.0f, 0.0f),
+				FVector2D(AllottedGeometry.GetLocalSize().X * ClampedFraction, AllottedGeometry.GetLocalSize().Y));
+
+			// Draw fill
+			FSlateDrawElement::MakeBox(
+				OutDrawElements,
+				RetLayerId++,
+				PaintRect,
+				FillImage,
+				DrawEffects,
+				FillColorAndOpacitySRGB
+			);
+
+			OutDrawElements.PopClip();
+		}
+		
+
+		RetLayerId++;
+
+		{
+			TArray<FVector2D> WarnLinePoints;
+			const double XPos = AllottedGeometry.GetLocalSize().X * WarnThreshold * XLimit;
+			WarnLinePoints.Add(FVector2D(XPos, 0.f));
+			WarnLinePoints.Add(FVector2D(XPos, AllottedGeometry.GetLocalSize().Y));
+
+			FSlateDrawElement::MakeLines(
+				OutDrawElements,
+				RetLayerId,
+				AllottedGeometry.ToPaintGeometry(),
+				WarnLinePoints,
+				DrawEffects,
+				(CachedFPSFraction < WarnThreshold ? FLinearColor(FColor::Orange) : FLinearColor::Black) * InWidgetStyle.GetColorAndOpacityTint()
+			);
+		}
+
+		{
+			TArray<FVector2D> TargetLinePoints;
+			const double XPos = AllottedGeometry.GetLocalSize().X * XLimit;
+			TargetLinePoints.Add(FVector2D(XPos, 0.f));
+			TargetLinePoints.Add(FVector2D(XPos, AllottedGeometry.GetLocalSize().Y));
+
+			FSlateDrawElement::MakeLines(
+				OutDrawElements,
+				RetLayerId,
+				AllottedGeometry.ToPaintGeometry(),
+				TargetLinePoints,
+				DrawEffects,
+				(CachedFPSFraction < 1.f ? FLinearColor::Red : FLinearColor::Black) * InWidgetStyle.GetColorAndOpacityTint(),
+				false,
+				2.f
+			);
+		}
+
+
+		return RetLayerId;
+	}
+
+	virtual FVector2D ComputeDesiredSize(float) const override
+	{
+		return { 20.0, 12.0 };
+	}
+
+	virtual bool ComputeVolatility() const override
+	{
+		return SLeafWidget::ComputeVolatility();
+	}
+
+	/** See attribute Percent */
+	void SetPercent(float InPercent)
+	{
+		if (!Percent.IdenticalTo(InPercent))
+		{
+			Percent = InPercent;
+			Invalidate(EInvalidateWidget::LayoutAndVolatility);
+		}
+	}
+
+	// Returns false if the clipping zone area is zero in which case we should skip drawing
+	bool PushTransformedClip(FSlateWindowElementList& OutDrawElements, const FGeometry& AllottedGeometry, FVector2D InsetPadding, FVector2D ProgressOrigin, FSlateRect Progress) const
+	{
+		const FSlateRenderTransform& Transform = AllottedGeometry.GetAccumulatedRenderTransform();
+
+		const FVector2D MaxSize = AllottedGeometry.GetLocalSize() - (InsetPadding * 2.0f);
+
+		const FSlateClippingZone ClippingZone(Transform.TransformPoint(InsetPadding + (ProgressOrigin - FVector2D(Progress.Left, Progress.Top)) * MaxSize),
+			Transform.TransformPoint(InsetPadding + FVector2D(ProgressOrigin.X + Progress.Right, ProgressOrigin.Y - Progress.Top) * MaxSize),
+			Transform.TransformPoint(InsetPadding + FVector2D(ProgressOrigin.X - Progress.Left, ProgressOrigin.Y + Progress.Bottom) * MaxSize),
+			Transform.TransformPoint(InsetPadding + (ProgressOrigin + FVector2D(Progress.Right, Progress.Bottom)) * MaxSize));
+
+		if (ClippingZone.HasZeroArea())
+		{
+			return false;
+		}
+
+		OutDrawElements.PushClip(ClippingZone);
+		return true;
+	}
+
+	float CachedFPSFraction = 0.f;
+	TAttribute<float> Percent = 0.f;
+	const FSlateBrush* FillImage = nullptr;
+	const FSlateBrush* BackgroundImage = nullptr;
+	TOptional<double> RefreshRate;
+	double LastUpdateTime = 0.0;
+};
\ No newline at end of file
diff --git a/Engine/Plugins/VirtualProduction/TimedDataMonitor/Source/TimedDataMonitorEditor/Private/STimedDataGenlock.cpp b/Engine/Plugins/VirtualProduction/TimedDataMonitor/Source/TimedDataMonitorEditor/Private/STimedDataGenlock.cpp
--- a/Engine/Plugins/VirtualProduction/TimedDataMonitor/Source/TimedDataMonitorEditor/Private/STimedDataGenlock.cpp
+++ b/Engine/Plugins/VirtualProduction/TimedDataMonitor/Source/TimedDataMonitorEditor/Private/STimedDataGenlock.cpp
@@ -12,7 +12,9 @@
 #include "Subsystems/AssetEditorSubsystem.h"
 
 #include "EditorFontGlyphs.h"
+#include "FixedFrameRateCustomTimeStep.h"
 #include "ScopedTransaction.h"
+#include "SFrameTime.h"
 #include "Styling/CoreStyle.h"
 
 #include "STimedDataMonitorPanel.h"
@@ -31,7 +33,7 @@ void STimedDataGenlock::Construct(const FArguments& InArgs, TSharedPtr<STimedDat
 {
 	OwnerPanel = InOwnerPanel;
 
-	UpdateCachedValue();
+	UpdateCachedValue(0.f);
 
 	FSlateFontInfo TimeFont = FCoreStyle::Get().GetFontStyle(TEXT("EmbossedText"));
 	TimeFont.Size += 4;
@@ -42,12 +44,11 @@ void STimedDataGenlock::Construct(const FArguments& InArgs, TSharedPtr<STimedDat
 		+ SVerticalBox::Slot()
 		.Padding(0)
 		.AutoHeight()
-		.HAlign(HAlign_Center)
 		[
 			SNew(SHorizontalBox)
 			+ SHorizontalBox::Slot()
-			.Padding(0.f, 0.f, 4.f, 0.f)
 			.VAlign(VAlign_Center)
+			.Padding(0.f, 0.f, 4.f, 0.f)
 			.AutoWidth()
 			[
 				SNew(STextBlock)
@@ -61,110 +62,97 @@ void STimedDataGenlock::Construct(const FArguments& InArgs, TSharedPtr<STimedDat
 			.AutoWidth()
 			[
 				SNew(STextBlock)
-				.MinDesiredWidth(500)
 				.Font(FCoreStyle::Get().GetFontStyle(TEXT("NormalText")))
 				.Text(this, &STimedDataGenlock::GetCustomTimeStepText)
 			]
+			+ SHorizontalBox::Slot()
+			.VAlign(VAlign_Bottom)
+			.AutoWidth()
+			[
+				SNew(SCheckBox)
+				.Padding(4.f)
+				.ToolTipText(LOCTEXT("ShowCustomTimeStepSetting_Tip", "Show custom time step provider setting"))
+				.Style(FAppStyle::Get(), "ToggleButtonCheckbox")
+				.ForegroundColor(FSlateColor::UseForeground())
+				.IsChecked_Lambda([](){return ECheckBoxState::Unchecked; })
+				.OnCheckStateChanged(this, &STimedDataGenlock::ShowCustomTimeStepSetting)
+				[
+					SNew(STextBlock)
+					.Font(FAppStyle::Get().GetFontStyle("FontAwesome.11"))
+					.Text(FEditorFontGlyphs::Cogs)
+				]
+			]
+			+ SHorizontalBox::Slot()
+			.AutoWidth()
+			.Padding(FMargin(4.f, 0.f, 0.f, 0.f))
+			[
+				SNew(SCheckBox)
+				.Padding(4.f)
+				.ToolTipText(LOCTEXT("ReapplyMenuToolTip", "Reinitialize the current Custom Time Step."))
+				.Style(FAppStyle::Get(), "ToggleButtonCheckbox")
+				.ForegroundColor(FSlateColor::UseForeground())
+				.IsEnabled(this, &STimedDataGenlock::IsCustomTimeStepEnabled)
+				.IsChecked_Lambda([]() {return ECheckBoxState::Unchecked; })
+				.OnCheckStateChanged(this, &STimedDataGenlock::ReinitializeCustomTimeStep)
+				[
+					SNew(STextBlock)
+					.Font(FAppStyle::Get().GetFontStyle("FontAwesome.11"))
+					.Text(FEditorFontGlyphs::Undo)
+				]
+			]
 		]
 		+ SVerticalBox::Slot()
 		.Padding(0)
 		.AutoHeight()
 		[
-			SNew(SProgressBar)
-			.BorderPadding(FVector2D::ZeroVector)
-			.Percent(this, &STimedDataGenlock::GetFPSFraction)
-			.FillColorAndOpacity(FSlateColor(FLinearColor(1.0f, 1.0f, 1.0f)))
+			SNew(SFrameTime)
 		]
 		+ SVerticalBox::Slot()
 		.Padding(0)
+		
 		.AutoHeight()
 		[
 			SNew(SGridPanel)
-			.FillColumn(0, 0.25f)
-			.FillColumn(1, 0.25f)
-			.FillColumn(2, 0.02f)
-			.FillColumn(4, 0.49f)
+			.FillColumn(0, 0.5f)
+			.FillColumn(1, 0.5f)
 
 			+ SGridPanel::Slot(0, 0)
+			.Padding(0.f, 2.f)
 			[
 				SNew(STextBlock)
 
 				.Text(LOCTEXT("FPSLabel", "FPS: "))
 			]
 			+ SGridPanel::Slot(1, 0)
+			.Padding(0.f, 2.f)
 			[
 				SNew(STextBlock)
 				.Text(this, &STimedDataGenlock::GetFPSText)
 			]
 			+ SGridPanel::Slot(0, 1)
+			.Padding(0.f, 2.f)
 			[
 				SNew(STextBlock)
 				.Text(LOCTEXT("DeltaTimeLabel", "DeltaTime: "))
 			]
 			+ SGridPanel::Slot(1, 1)
+			.Padding(0.f, 2.f)
 			[
 				SNew(STextBlock)
 				.Text(this, &STimedDataGenlock::GetDeltaTimeText)
 			]
 			+ SGridPanel::Slot(0, 2)
+			.Padding(0.f, 2.f)
 			[
 				SNew(STextBlock)
 				.Text(LOCTEXT("IdleTimeLabel", "Idle Time: "))
 			]
 			+ SGridPanel::Slot(1, 2)
+			.Padding(0.f, 2.f)
 			[
 				SNew(STextBlock)
 				.Text(this, &STimedDataGenlock::GetIdleTimeText)
 			]
-			+ SGridPanel::Slot(3, 0)
-			.RowSpan(3)
-			[
-				SNew(SSeparator)
-				.Thickness(2)
-				.Orientation(EOrientation::Orient_Vertical)
-			]
-			+ SGridPanel::Slot(4, 0)
-			.RowSpan(3)
-			.HAlign(HAlign_Right)
-			.VAlign(VAlign_Bottom)
-			[
-				SNew(SHorizontalBox)
-				+ SHorizontalBox::Slot()
-				.VAlign(VAlign_Bottom)
-				.AutoWidth()
-				[
-					SNew(SCheckBox)
-					.Padding(4.f)
-					.ToolTipText(LOCTEXT("ShowTimecodeProviderSetting_Tip", "Show timecode provider setting"))
-					.Style(FAppStyle::Get(), "ToggleButtonCheckbox")
-					.ForegroundColor(FSlateColor::UseForeground())
-					.IsChecked_Lambda([](){return ECheckBoxState::Unchecked; })
-					.OnCheckStateChanged(this, &STimedDataGenlock::ShowCustomTimeStepSetting)
-					[
-						SNew(STextBlock)
-						.Font(FAppStyle::Get().GetFontStyle("FontAwesome.11"))
-						.Text(FEditorFontGlyphs::Cogs)
-					]
-				]
-				+ SHorizontalBox::Slot()
-				.AutoWidth()
-				.Padding(FMargin(4.f, 0.f, 0.f, 0.f))
-				[
-					SNew(SCheckBox)
-					.Padding(4.f)
-					.ToolTipText(LOCTEXT("ReapplyMenuToolTip", "Reinitialize the current Custom Time Step."))
-					.Style(FAppStyle::Get(), "ToggleButtonCheckbox")
-					.ForegroundColor(FSlateColor::UseForeground())
-					.IsEnabled(this, &STimedDataGenlock::IsCustomTimeStepEnabled)
-					.IsChecked_Lambda([]() {return ECheckBoxState::Unchecked; })
-					.OnCheckStateChanged(this, &STimedDataGenlock::ReinitializeCustomTimeStep)
-					[
-						SNew(STextBlock)
-						.Font(FAppStyle::Get().GetFontStyle("FontAwesome.11"))
-						.Text(FEditorFontGlyphs::Undo)
-					]
-				]
-			]
 		]
 	];
 }
@@ -178,10 +166,19 @@ void STimedDataGenlock::RequestRefresh()
 	}
 }
 
-
-void STimedDataGenlock::UpdateCachedValue()
+void STimedDataGenlock::UpdateCachedValue(float InDeltaTime)
 {
-	CachedFPSFraction = FApp::GetDeltaTime() > 0.0 ? (1.f-float(FApp::GetIdleTime() / FApp::GetDeltaTime())) : 1.f;
+	double IdleTime = FApp::GetIdleTime();
+
+	const float MaxTickRate = GEngine->GetMaxTickRate(0.001f, false);
+
+	float TargetFPS = GEngine->bUseFixedFrameRate ? GEngine->FixedFrameRate : MaxTickRate;
+	if (UFixedFrameRateCustomTimeStep* CustomTimeStep = Cast<UFixedFrameRateCustomTimeStep>(GEngine->GetCustomTimeStep()))
+	{
+		TargetFPS = CustomTimeStep->GetFixedFrameRate().AsDecimal();
+	}
+	
+	const float TargetFrameTime = 1 / TargetFPS;
 
 	FNumberFormattingOptions FPSFormattingOptions = FNumberFormattingOptions()
 		.SetUseGrouping(false)
@@ -190,9 +187,16 @@ void STimedDataGenlock::UpdateCachedValue()
 		.SetUseGrouping(false)
 		.SetMinimumFractionalDigits(3)
 		.SetMaximumFractionalDigits(3);
+
+	uint32 IdleFrames = 0;
+	if (!FMath::IsNearlyZero(IdleTime))
+	{
+		IdleFrames = FMath::RoundToInt(IdleTime / (TargetFrameTime * TargetFrameTime) );
+	}
+
 	CachedFPSText = FText::AsNumber(1.0/FApp::GetDeltaTime(), &FPSFormattingOptions);
 	CachedDeltaTimeText = FText::Format(LOCTEXT("WithMilliSeconds", "{0}ms"), FText::AsNumber(FApp::GetDeltaTime()*1000, &MsFormattingOptions));
-	CachedIdleTimeText = FText::Format(LOCTEXT("WithMilliSeconds", "{0}ms"), FText::AsNumber(FApp::GetIdleTime()*1000, &MsFormattingOptions));
+	CachedIdleTimeText = FText::Format(LOCTEXT("WithMilliSecondsFrames", "{0}ms ({1} frames)"), FText::AsNumber(IdleTime *1000, &MsFormattingOptions), IdleFrames);
 
 	if (const UEngineCustomTimeStep* CustomTimeStep = GEngine->GetCustomTimeStep())
 	{
@@ -202,18 +206,12 @@ void STimedDataGenlock::UpdateCachedValue()
 	else
 	{
 		CachedState = ECustomTimeStepSynchronizationState::Error;
-		CachedCustomTimeStepText = LOCTEXT("Undefined", "<Default Engine Settings>");
+		CachedCustomTimeStepText = LOCTEXT("NoCustomTimeStep", "No custom time step");
 	}
 }
 
-
 FText STimedDataGenlock::GetStateText() const
 {
-	if (!IsCustomTimeStepEnabled())
-	{
-		return FText::GetEmpty();
-	}
-
 	switch (CachedState)
 	{
 	case ECustomTimeStepSynchronizationState::Synchronized:
@@ -256,13 +254,6 @@ bool STimedDataGenlock::IsCustomTimeStepEnabled() const
 	return GEngine && GEngine->GetCustomTimeStep() != nullptr;
 }
 
-
-TOptional<float> STimedDataGenlock::GetFPSFraction() const
-{
-	return CachedFPSFraction;
-}
-
-
 void STimedDataGenlock::ShowCustomTimeStepSetting(ECheckBoxState)
 {
 	if (UEngineCustomTimeStep* CustomTimeStepPtr = GEngine->GetCustomTimeStep())
diff --git a/Engine/Plugins/VirtualProduction/TimedDataMonitor/Source/TimedDataMonitorEditor/Private/STimedDataGenlock.h b/Engine/Plugins/VirtualProduction/TimedDataMonitor/Source/TimedDataMonitorEditor/Private/STimedDataGenlock.h
--- a/Engine/Plugins/VirtualProduction/TimedDataMonitor/Source/TimedDataMonitorEditor/Private/STimedDataGenlock.h
+++ b/Engine/Plugins/VirtualProduction/TimedDataMonitor/Source/TimedDataMonitorEditor/Private/STimedDataGenlock.h
@@ -22,7 +22,7 @@ public:
 	void Construct(const FArguments& InArgs, TSharedPtr<STimedDataMonitorPanel> OwnerPanel);
 
 	void RequestRefresh();
-	void UpdateCachedValue();
+	void UpdateCachedValue(float InDeltaTime);
 
 private:
 	FText GetStateText() const;
@@ -30,7 +30,7 @@ private:
 	FText GetCustomTimeStepText() const;
 
 	bool IsCustomTimeStepEnabled() const;
-	TOptional<float> GetFPSFraction() const;
+	float GetFPSFraction() const;
 
 	FText GetFPSText() const  { return CachedFPSText; }
 	FText GetDeltaTimeText() const  { return CachedDeltaTimeText; }
@@ -47,5 +47,4 @@ private:
 	FText CachedFPSText;
 	FText CachedDeltaTimeText;
 	FText CachedIdleTimeText;
-	TOptional<float> CachedFPSFraction;
 };
diff --git a/Engine/Plugins/VirtualProduction/TimedDataMonitor/Source/TimedDataMonitorEditor/Private/STimedDataListView.cpp b/Engine/Plugins/VirtualProduction/TimedDataMonitor/Source/TimedDataMonitorEditor/Private/STimedDataListView.cpp
--- a/Engine/Plugins/VirtualProduction/TimedDataMonitor/Source/TimedDataMonitorEditor/Private/STimedDataListView.cpp
+++ b/Engine/Plugins/VirtualProduction/TimedDataMonitor/Source/TimedDataMonitorEditor/Private/STimedDataListView.cpp
@@ -2,29 +2,28 @@
 
 #include "STimedDataListView.h"
 
+#include "EditorFontGlyphs.h"
 #include "Engine/Engine.h"
+#include "Framework/MultiBox/MultiBoxBuilder.h"
 #include "ITimedDataInput.h"
 #include "Misc/App.h"
 #include "Misc/DateTime.h"
 #include "Misc/Timecode.h"
 #include "Misc/Timespan.h"
 #include "TimedDataMonitorEditorSettings.h"
-#include "TimedDataMonitorSubsystem.h"
-
-#include "EditorFontGlyphs.h"
-#include "Styling/AppStyle.h"
 #include "TimedDataMonitorEditorStyle.h"
-
-#include "Framework/MultiBox/MultiBoxBuilder.h"
+#include "TimedDataMonitorSubsystem.h"
+#include "SEnumCombo.h"
 #include "STimedDataMonitorPanel.h"
 #include "STimedDataNumericEntryBox.h"
 #include "STimingDiagramWidget.h"
+#include "Styling/AppStyle.h"
 #include "Widgets/Images/SImage.h"
 #include "Widgets/Input/SCheckBox.h"
 #include "Widgets/Input/SNumericEntryBox.h"
+#include "Widgets/SNullWidget.h"
 #include "Widgets/Text/STextBlock.h"
 #include "Widgets/Views/SExpanderArrow.h"
-#include "Widgets/SNullWidget.h"
 
 
 #define LOCTEXT_NAMESPACE "STimedDataListView"
@@ -66,9 +65,19 @@ struct FTimedDataInputTableRowData : TSharedFromThis<FTimedDataInputTableRowData
 
 		DisplayName = TimedDataMonitorSubsystem->GetInputDisplayName(InputIdentifier);
 
-		if (ITimedDataInput* InpuTimedData = TimedDataMonitorSubsystem->GetTimedDataInput(InputIdentifier))
+		ETimedDataInputEvaluationType InputType = TimedDataMonitorSubsystem->GetInputEvaluationType(InputIdentifier);
+		if (InputType == ETimedDataInputEvaluationType::Timecode)
+		{
+			CachedInputEvaluationOffsetType = ETimedDataInputEvaluationOffsetType::Frames;
+		}
+		else
+		{
+			CachedInputEvaluationOffsetType = ETimedDataInputEvaluationOffsetType::Seconds;
+		}
+
+		if (ITimedDataInput* InputTimedData = TimedDataMonitorSubsystem->GetTimedDataInput(InputIdentifier))
 		{
-			InputIcon = InpuTimedData->GetDisplayIcon();
+			InputIcon = InputTimedData->GetDisplayIcon();
 		}
 	}
 
@@ -105,7 +114,9 @@ struct FTimedDataInputTableRowData : TSharedFromThis<FTimedDataInputTableRowData
 			};
 
 			CachedInputEvaluationType = TimedDataMonitorSubsystem->GetInputEvaluationType(InputIdentifier);
-			CachedInputEvaluationOffset = TimedDataMonitorSubsystem->GetInputEvaluationOffsetInSeconds(InputIdentifier);
+
+			CachedInputEvaluationOffsetSeconds = TimedDataMonitorSubsystem->GetInputEvaluationOffsetInSeconds(InputIdentifier);
+			CachedInputEvaluationOffsetFrames = TimedDataMonitorSubsystem->GetInputEvaluationOffsetInFrames(InputIdentifier);
 			CachedState = TimedDataMonitorSubsystem->GetInputConnectionState(InputIdentifier);
 			CachedBufferSize = TimedDataMonitorSubsystem->GetInputDataBufferSize(InputIdentifier);
 			CachedCurrentAmountOfBuffer = 0;
@@ -132,7 +143,8 @@ struct FTimedDataInputTableRowData : TSharedFromThis<FTimedDataInputTableRowData
 		{
 			CachedEnabled = TimedDataMonitorSubsystem->IsChannelEnabled(ChannelIdentifier) ? ECheckBoxState::Checked : ECheckBoxState::Unchecked;
 			CachedInputEvaluationType = TimedDataMonitorSubsystem->GetInputEvaluationType(InputIdentifier);
-			CachedInputEvaluationOffset = 0.f;
+			CachedInputEvaluationOffsetSeconds = 0.f;
+			CachedInputEvaluationOffsetFrames = 0.f;
 			CachedState = TimedDataMonitorSubsystem->GetChannelConnectionState(ChannelIdentifier);
 			CachedBufferSize = TimedDataMonitorSubsystem->GetChannelDataBufferSize(ChannelIdentifier);
 			CachedCurrentAmountOfBuffer = TimedDataMonitorSubsystem->GetChannelNumberOfSamples(ChannelIdentifier);
@@ -180,7 +192,10 @@ public:
 
 	ECheckBoxState CachedEnabled = ECheckBoxState::Undetermined;
 	ETimedDataInputEvaluationType CachedInputEvaluationType = ETimedDataInputEvaluationType::None;
-	float CachedInputEvaluationOffset = 0.f;
+	ETimedDataInputEvaluationOffsetType CachedInputEvaluationOffsetType = ETimedDataInputEvaluationOffsetType::Seconds;
+	float CachedInputEvaluationOffsetSeconds = 0.f;
+	float CachedInputEvaluationOffsetFrames = 0.f;
+
 	ETimedDataInputState CachedState = ETimedDataInputState::Disconnected;
 	FText CachedDescription;
 	int32 CachedBufferSize = 0;
@@ -319,12 +334,17 @@ TSharedRef<SWidget> STimedDataInputTableRow::GenerateWidgetForColumn(const FName
 	{
 		if (Item->bIsInput)
 		{
+			TSharedRef<SWidget> Label = SNew(SEnumComboBox, StaticEnum<ETimedDataInputEvaluationOffsetType>())
+				.ContentPadding(0.f)
+				.CurrentValue_Lambda([this]() { return static_cast<int32>(Item->CachedInputEvaluationOffsetType); })
+				.OnEnumSelectionChanged(this, &STimedDataInputTableRow::OnEvaluationOffsetTypeChanged);
+
 			return SNew(STimedDataNumericEntryBox<float>)
 				.TextStyle(ItemTextBlockStyle)
-				.MinValue(1)
 				.ToolTipText(LOCTEXT("TimeCorrection_ToolTip", "Time Correction."))
 				.Value(this, &STimedDataInputTableRow::GetEvaluationOffset)
-				.EditLabel(LOCTEXT("TimeCorrection_EditLable", "In Seconds: "))
+				.ComboButton(false)
+				.SuffixWidget(Label)
 				.OnValueCommitted(this, &STimedDataInputTableRow::SetEvaluationOffset);
 		}
 		return SNullWidget::NullWidget;
@@ -335,7 +355,6 @@ TSharedRef<SWidget> STimedDataInputTableRow::GenerateWidgetForColumn(const FName
 		{
 			return SNew(STimedDataNumericEntryBox<int32>)
 				.TextStyle(ItemTextBlockStyle)
-				.MinValue(1)
 				.ToolTipText(LOCTEXT("BufferSize_ToolTip", "Buffer Size."))
 				.Value(this, &STimedDataInputTableRow::GetBufferSize)
 				.ShowAmount(!Item->bIsInput)
@@ -468,7 +487,7 @@ FText STimedDataInputTableRow::GetEvaluationOffsetText() const
 {
 	if (Item->bIsInput)
 	{
-		return FText::AsNumber(Item->CachedInputEvaluationOffset);
+		return Item->CachedInputEvaluationOffsetType == ETimedDataInputEvaluationOffsetType::Seconds ? FText::AsNumber(Item->CachedInputEvaluationOffsetSeconds) : FText::AsNumber(Item->CachedInputEvaluationOffsetFrames);
 	}
 	return FText::GetEmpty();
 }
@@ -476,7 +495,7 @@ FText STimedDataInputTableRow::GetEvaluationOffsetText() const
 
 float STimedDataInputTableRow::GetEvaluationOffset() const
 {
-	return Item->CachedInputEvaluationOffset;
+	return Item->CachedInputEvaluationOffsetType == ETimedDataInputEvaluationOffsetType::Seconds ? Item->CachedInputEvaluationOffsetSeconds : Item->CachedInputEvaluationOffsetFrames;
 }
 
 
@@ -487,7 +506,21 @@ void STimedDataInputTableRow::SetEvaluationOffset(float InValue, ETextCommit::Ty
 		UTimedDataMonitorSubsystem* TimedDataMonitorSubsystem = GEngine->GetEngineSubsystem<UTimedDataMonitorSubsystem>();
 		check(TimedDataMonitorSubsystem);
 
-		TimedDataMonitorSubsystem->SetInputEvaluationOffsetInSeconds(Item->InputIdentifier, InValue);
+		if (Item->CachedInputEvaluationOffsetType == ETimedDataInputEvaluationOffsetType::Seconds)
+		{
+			Item->CachedInputEvaluationOffsetSeconds = InValue;
+			TimedDataMonitorSubsystem->SetInputEvaluationOffsetInSeconds(Item->InputIdentifier, InValue);
+		}
+		else if (Item->CachedInputEvaluationOffsetType == ETimedDataInputEvaluationOffsetType::Frames)
+		{
+			Item->CachedInputEvaluationOffsetFrames = InValue;
+			TimedDataMonitorSubsystem->SetInputEvaluationOffsetInFrames(Item->InputIdentifier, InValue);
+		}
+		else
+		{
+			ensure(false);
+		}
+		
 		OwnerTreeView->RequestRefresh();
 	}
 }
@@ -631,6 +664,16 @@ FText STimedDataInputTableRow::GetFrameDroppedCount() const
 	return FText::AsNumber(Item->CachedStatsFrameDropped);
 }
 
+void STimedDataInputTableRow::OnEvaluationOffsetTypeChanged(int32 NewValue, ESelectInfo::Type)
+{
+	FSlateApplication::Get().ForEachUser([&](FSlateUser& User)
+	{
+		User.ClearFocus(EFocusCause::SetDirectly);
+	});
+
+	Item->CachedInputEvaluationOffsetType = static_cast<ETimedDataInputEvaluationOffsetType>(NewValue);
+}
+
 /**
  * STimedDataListView
  */
@@ -690,8 +733,8 @@ void STimedDataInputListView::Construct(const FArguments& InArgs, TSharedPtr<STi
 			.DefaultTooltip(LOCTEXT("HeaderTooltip_EvaluationMode", "How the input is evaluated"))
 
 			+ SHeaderRow::Column(TimedDataListView::HeaderIdName_TimeCorrection)
-			.FixedWidth(100)
-			.HAlignCell(EHorizontalAlignment::HAlign_Right)
+			.FixedWidth(170)
+			.HAlignCell(EHorizontalAlignment::HAlign_Center)
 			.DefaultLabel(LOCTEXT("HeaderName_TimeCorrection", "Time Correction"))
 
 			+ SHeaderRow::Column(TimedDataListView::HeaderIdName_BufferSize)
@@ -801,6 +844,16 @@ void STimedDataInputListView::RebuildSources()
 		}
 	}
 
+	struct FSortNamesAlphabetically
+	{
+		bool operator()(const TSharedPtr<FTimedDataInputTableRowData>& LHS, const TSharedPtr<FTimedDataInputTableRowData>& RHS) const
+		{
+			return (LHS->DisplayName.CompareTo(RHS->DisplayName) < 0);
+		}
+	};
+
+	ListItemsSource.Sort(FSortNamesAlphabetically());
+
 	for (FTimedDataInputTableRowDataPtr& TableRowData : ListItemsSource)
 	{
 		TableRowData->UpdateCachedValue();
diff --git a/Engine/Plugins/VirtualProduction/TimedDataMonitor/Source/TimedDataMonitorEditor/Private/STimedDataListView.h b/Engine/Plugins/VirtualProduction/TimedDataMonitor/Source/TimedDataMonitorEditor/Private/STimedDataListView.h
--- a/Engine/Plugins/VirtualProduction/TimedDataMonitor/Source/TimedDataMonitorEditor/Private/STimedDataListView.h
+++ b/Engine/Plugins/VirtualProduction/TimedDataMonitor/Source/TimedDataMonitorEditor/Private/STimedDataListView.h
@@ -13,6 +13,8 @@
 #include "Widgets/Views/STreeView.h"
 #include "Widgets/Views/STableRow.h"
 
+#include "STimedDataListView.generated.h"
+
 enum class ETimedDataInputEvaluationType : uint8;
 struct FTimedDataInputTableRowData;
 struct FSlateBrush;
@@ -21,6 +23,15 @@ class STimedDataInputTableRow;
 class STimedDataMonitorPanel;
 class STimingDiagramWidget;
 
+UENUM()
+enum class ETimedDataInputEvaluationOffsetType : uint8
+{
+	/** The input offset is specified in seconds. */
+	Seconds UMETA(DisplayName = "s"),
+	/** The input offset is specified in frames. */
+	Frames UMETA(DisplayName = "f"),
+};
+
 
 /**
  * 
@@ -72,6 +83,8 @@ private:
 	FText GetBufferOverflowCount() const;
 	FText GetFrameDroppedCount() const;
 
+	void OnEvaluationOffsetTypeChanged(int32 NewValue, ESelectInfo::Type);
+
 private:
 	FTimedDataInputTableRowDataPtr Item;
 	TSharedPtr<STimedDataInputListView> OwnerTreeView;
diff --git a/Engine/Plugins/VirtualProduction/TimedDataMonitor/Source/TimedDataMonitorEditor/Private/STimedDataMonitorBufferVisualizer.cpp b/Engine/Plugins/VirtualProduction/TimedDataMonitor/Source/TimedDataMonitorEditor/Private/STimedDataMonitorBufferVisualizer.cpp
--- a/Engine/Plugins/VirtualProduction/TimedDataMonitor/Source/TimedDataMonitorEditor/Private/STimedDataMonitorBufferVisualizer.cpp
+++ b/Engine/Plugins/VirtualProduction/TimedDataMonitor/Source/TimedDataMonitorEditor/Private/STimedDataMonitorBufferVisualizer.cpp
@@ -260,10 +260,15 @@ void STimedDataMonitorBufferVisualizer::RebuiltListItemsSource()
 
 	if (bListModified)
 	{
-		Algo::Sort(ListItemsSource, [](const TSharedPtr<FTimedDataMonitorBuffervisualizerItem>& A, const TSharedPtr<FTimedDataMonitorBuffervisualizerItem>& B)
+		struct FSortNamesAlphabetically
+		{
+			bool operator()(const TSharedPtr<FTimedDataMonitorBuffervisualizerItem>& LHS, const TSharedPtr<FTimedDataMonitorBuffervisualizerItem>& RHS) const
 			{
-				return (A->InputIdentifier == B->InputIdentifier) ? A->ChannelDisplayName.CompareTo(B->ChannelDisplayName) < 0 : A->InputIdentifier < B->InputIdentifier;
-			});
+				return (LHS->ChannelDisplayName.CompareTo(RHS->ChannelDisplayName) < 0);
+			}
+		};
+
+		ListItemsSource.Sort(FSortNamesAlphabetically());
 		ListViewWidget->RequestListRefresh();
 	}
 }
diff --git a/Engine/Plugins/VirtualProduction/TimedDataMonitor/Source/TimedDataMonitorEditor/Private/STimedDataMonitorPanel.cpp b/Engine/Plugins/VirtualProduction/TimedDataMonitor/Source/TimedDataMonitorEditor/Private/STimedDataMonitorPanel.cpp
--- a/Engine/Plugins/VirtualProduction/TimedDataMonitor/Source/TimedDataMonitorEditor/Private/STimedDataMonitorPanel.cpp
+++ b/Engine/Plugins/VirtualProduction/TimedDataMonitor/Source/TimedDataMonitorEditor/Private/STimedDataMonitorPanel.cpp
@@ -186,9 +186,10 @@ void STimedDataMonitorPanel::Construct(const FArguments& InArgs)
 			];
 	}
 
-	TSharedRef<SWidget> Toolbar = SNew(SBorder)
-		.VAlign(VAlign_Center)
-		.Padding(FMargin(5.f, 5.f, 5.f, 5.f))
+	TSharedRef<SWidget> Toolbar = 
+		SNew(SBorder)
+		.BorderImage(FCoreStyle::Get().GetBrush("ToolPanel.GroupBorder"))
+		.Padding(FMargin(7.f, 5.f))
 		[
 			SNew(SHorizontalBox)
 			// Calibrate button
@@ -329,7 +330,7 @@ void STimedDataMonitorPanel::Construct(const FArguments& InArgs)
 			+ SHorizontalBox::Slot()
 			.AutoWidth()
 			.VAlign(VAlign_Center)
-			.Padding(4.f)
+			.Padding(4.f, 4.f, 6.f, 4.f)
 			[
 				SNew(STextBlock)
 				.MinDesiredWidth(100)
@@ -339,108 +340,115 @@ void STimedDataMonitorPanel::Construct(const FArguments& InArgs)
 			]
 		];
 
+	FSlateFontInfo SectionHeaderFont = FCoreStyle::Get().GetFontStyle(TEXT("EmbossedText"));
+	SectionHeaderFont.Size += 4;
+
 	ChildSlot
 	[
 		SNew(SOverlay)
 		+ SOverlay::Slot()
 		[
-			SNew(SVerticalBox)
-			// toolbar
-			+ SVerticalBox::Slot()
-			.AutoHeight()
-			.Padding(10.f, 10.f, 10.f, 10.f)
-			[
-				Toolbar
-			]
-			// timing element
-			+ SVerticalBox::Slot()
-			.AutoHeight()
-			.Padding(10.f, 0.f, 10.f, 10.f)
+			SNew(SBorder)
+			.BorderImage(FCoreStyle::Get().GetBrush("ToolPanel.GroupBorder"))
+			.BorderBackgroundColor(FLinearColor::Gray) // Darken the outer border
+			.Padding(0.f)
 			[
-				SNew(SHorizontalBox)
-				+ SHorizontalBox::Slot()
-				.FillWidth(0.5f)
-				.Padding(0.f, 0.f, 10.f, 0.f)
+				SNew(SVerticalBox)
+				// Toolbar
+				+ SVerticalBox::Slot()
+				.AutoHeight()
 				[
-					SNew(SVerticalBox)
-					+ SVerticalBox::Slot()
-					.AutoHeight()
+					Toolbar
+				]
+				// Timing element
+				+ SVerticalBox::Slot()
+				.Padding(4.f)
+				.AutoHeight()
+				[
+					SNew(SHorizontalBox)
+					+ SHorizontalBox::Slot()
+					.FillWidth(0.5f)
+					.Padding(FMargin(0.f, 0.f, 2.5f, 0.f))
 					[
-						SNew(STextBlock)
-						.Text(LOCTEXT("GenlockLabel", "Genlock"))
+						SNew(SBorder)
+						.BorderImage(FCoreStyle::Get().GetBrush("ToolPanel.GroupBorder"))
+						.Padding(FMargin(5.f, 2.f))
+						[
+							SNew(SVerticalBox)
+							+ SVerticalBox::Slot()
+							.AutoHeight()
+							[
+								SNew(STextBlock)
+								.Font(SectionHeaderFont)
+								.Text(LOCTEXT("GenlockLabel", "Genlock"))
+							]
+							+ SVerticalBox::Slot()
+							.AutoHeight()
+							[
+								SAssignNew(TimedDataGenlockWidget, STimedDataGenlock, SharedThis<STimedDataMonitorPanel>(this))
+							]
+						]
 					]
-					+ SVerticalBox::Slot()
-					.AutoHeight()
+					+ SHorizontalBox::Slot()
+					.FillWidth(0.5f)
+					.Padding(FMargin(2.5f, 0.f, 0.f, 0.f))
 					[
 						SNew(SBorder)
-						.HAlign(HAlign_Center)
-						.VAlign(VAlign_Center)
-						.Padding(FMargin(4.f, 4.f, 4.f, 4.f))
+						.BorderImage(FCoreStyle::Get().GetBrush("ToolPanel.GroupBorder"))
+						.Padding(FMargin(5.f, 2.f))
 						[
-							SAssignNew(TimedDataGenlockWidget, STimedDataGenlock, SharedThis<STimedDataMonitorPanel>(this))
+							SNew(SVerticalBox)
+							+ SVerticalBox::Slot()
+							.AutoHeight()
+							[
+								SNew(STextBlock)
+								.Font(SectionHeaderFont)
+								.Text(LOCTEXT("TimecodeLabel", "Timecode Provider"))
+							]
+							+ SVerticalBox::Slot()
+							[
+								SAssignNew(TimedDataTimecodeWidget, STimedDataTimecodeProvider, SharedThis<STimedDataMonitorPanel>(this))
+							]
 						]
 					]
 				]
-				+ SHorizontalBox::Slot()
-				.FillWidth(0.5f)
-				.Padding(0.f)
+				// Sources list
+				+ SVerticalBox::Slot()
+				.Padding(5.f, 2.5f)
+				.FillHeight(1.f)
 				[
-					SNew(SVerticalBox)
-					+ SVerticalBox::Slot()
-					.AutoHeight()
+					SNew(SScrollBox)
+					.Orientation(Orient_Vertical)
+					+ SScrollBox::Slot()
 					[
-						SNew(STextBlock)
-						.Text(LOCTEXT("TimecodeLabel", "Timecode Provider"))
+						SAssignNew(TimedDataSourceList, STimedDataInputListView, SharedThis<STimedDataMonitorPanel>(this))
+						.OnContextMenuOpening(this, &STimedDataMonitorPanel::OnDataListConstructContextMenu)
 					]
-					+ SVerticalBox::Slot()
-					.AutoHeight()
+				]
+				// Message log
+				+ SVerticalBox::Slot()
+				.AutoHeight()
+				[
+					SNew(SBorder)
+					.Padding(FMargin(5.f, 5.f, 5.f, 5.f))
+					.Visibility(this, &STimedDataMonitorPanel::ShowMessageLog)
 					[
-						SNew(SBorder)
-						.HAlign(HAlign_Center)
-						.VAlign(VAlign_Center)
+						SNew(SHorizontalBox)
+						+ SHorizontalBox::Slot()
+						.FillWidth(1.0f)
 						[
-							SAssignNew(TimedDataTimecodeWidget, STimedDataTimecodeProvider, SharedThis<STimedDataMonitorPanel>(this))
+							MessageLogListingWidget
 						]
 					]
 				]
-			]
-			// sources list
-			+ SVerticalBox::Slot()
-			.FillHeight(1.f)
-			.Padding(10.f, 0.f, 10.f, 10.f)
-			[
-				SNew(SScrollBox)
-				.Orientation(Orient_Vertical)
-				+ SScrollBox::Slot()
-				[
-					SAssignNew(TimedDataSourceList, STimedDataInputListView, SharedThis<STimedDataMonitorPanel>(this))
-					.OnContextMenuOpening(this, &STimedDataMonitorPanel::OnDataListConstructContextMenu)
-				]
-			]
-			// message log
-			+ SVerticalBox::Slot()
-			.AutoHeight()
-			.Padding(10.f, 0.f, 10.f, 10.f)
-			[
-				SNew(SBorder)
-				.Padding(FMargin(5.f, 5.f, 5.f, 5.f))
-				.Visibility(this, &STimedDataMonitorPanel::ShowMessageLog)
+				// Show performance warning
+				+SVerticalBox::Slot()
+				.AutoHeight()
+				.Padding(10.f, 0.f, 10.f, 10.f)
 				[
-					SNew(SHorizontalBox)
-					+ SHorizontalBox::Slot()
-					.FillWidth(1.0f)
-					[
-						MessageLogListingWidget
-					]
+					PerformanceWidget.ToSharedRef()
 				]
 			]
-			// show performance warning
-			+SVerticalBox::Slot()
-			.AutoHeight()
-			.Padding(10.f, 0.f, 10.f, 10.f)
-			[
-				PerformanceWidget.ToSharedRef()
-			]
 		]
 		+ SOverlay::Slot()
 		[
@@ -498,7 +506,7 @@ void STimedDataMonitorPanel::Tick(const FGeometry& AllottedGeometry, const doubl
 
 		if (TimedDataGenlockWidget)
 		{
-			TimedDataGenlockWidget->UpdateCachedValue();
+			TimedDataGenlockWidget->UpdateCachedValue(InDeltaTime);
 		}
 		if (TimedDataTimecodeWidget)
 		{
diff --git a/Engine/Plugins/VirtualProduction/TimedDataMonitor/Source/TimedDataMonitorEditor/Private/STimedDataNumericEntryBox.h b/Engine/Plugins/VirtualProduction/TimedDataMonitor/Source/TimedDataMonitorEditor/Private/STimedDataNumericEntryBox.h
--- a/Engine/Plugins/VirtualProduction/TimedDataMonitor/Source/TimedDataMonitorEditor/Private/STimedDataNumericEntryBox.h
+++ b/Engine/Plugins/VirtualProduction/TimedDataMonitor/Source/TimedDataMonitorEditor/Private/STimedDataNumericEntryBox.h
@@ -18,6 +18,7 @@
 #include "Widgets/Input/SButton.h"
 #include "Widgets/Input/SNumericEntryBox.h"
 #include "Widgets/Layout/SBorder.h"
+#include "Widgets/SNullWidget.h"
 #include "Widgets/Text/STextBlock.h"
 
 
@@ -32,18 +33,23 @@ private:
 public:
 	SLATE_BEGIN_ARGS(STimedDataNumericEntryBox)
 			: _Value(0)
-			, _MinValue(0)
+			, _SuffixWidget(SNullWidget::NullWidget)
 			, _CanEdit(true)
 			, _Amount(0)
 			, _ShowAmount(false)
+			, _TextStyle(&FAppStyle::Get().GetWidgetStyle<FTextBlockStyle>(TEXT("NormalText")))
+			, _ComboButton(true)
 		{ }
 		SLATE_ATTRIBUTE(NumericType, Value)
-		SLATE_ARGUMENT(NumericType, MinValue)
+		SLATE_ARGUMENT(TOptional<NumericType>, MinValue)
+		SLATE_ARGUMENT(TOptional<NumericType>, MaxValue)
 		SLATE_ARGUMENT(FText, EditLabel)
+		SLATE_ARGUMENT(TSharedRef<SWidget>, SuffixWidget)
 		SLATE_ARGUMENT(bool, CanEdit);
 		SLATE_ATTRIBUTE(NumericType, Amount);
 		SLATE_ARGUMENT(bool, ShowAmount);
 		SLATE_STYLE_ARGUMENT(FTextBlockStyle, TextStyle)
+		SLATE_ARGUMENT(bool, ComboButton)
 		SLATE_EVENT(FOnValueCommitted, OnValueCommitted)
 	SLATE_END_ARGS()
 	
@@ -51,56 +57,96 @@ public:
 	void Construct(const FArguments& InArgs)
 	{
 		Value = InArgs._Value;
+		CachedValue = Value.Get();
 		MinValue = InArgs._MinValue;
-		EditLabel = InArgs._EditLabel;
+		MaxValue = InArgs._MaxValue;
+		SuffixWidget = InArgs._SuffixWidget;
 		Amount = InArgs._Amount;
 		bShowAmount = InArgs._ShowAmount;
+		bComboButton = InArgs._ComboButton;
 		OnValueCommitted = InArgs._OnValueCommitted;
 
-		TSharedRef<STextBlock> TextBlock = SNew(STextBlock)
+		TSharedRef<SWidget> TextBlock = SNew(STextBlock)
 			.TextStyle(InArgs._TextStyle)
 			.Text(this, &STimedDataNumericEntryBox::GetValueText);
 
-		if (InArgs._CanEdit)
+		if (InArgs._ComboButton)
 		{
-			ChildSlot
-			[
-				SNew(SHorizontalBox)
-				+ SHorizontalBox::Slot()
-				.HAlign(HAlign_Right)
-				.VAlign(VAlign_Center)
-				.FillWidth(1.f)
-				.Padding(4.0f, 0.0f, 0.0f, 0.0f)
-				[
-					TextBlock
-				]
-				+ SHorizontalBox::Slot()
-				.HAlign(HAlign_Right)
-				.VAlign(VAlign_Center)
-				.AutoWidth()
-				.Padding(4.0f, 0.0f, 0.0f, 0.0f)
+			if (InArgs._CanEdit)
+			{
+				ChildSlot
 				[
-					SAssignNew(ComboButton, SComboButton)
-					.ComboButtonStyle(FTimedDataMonitorEditorStyle::Get(), "ToggleComboButton")
-					.HAlign(HAlign_Center)
-					.HasDownArrow(false)
-					.OnGetMenuContent(this, &STimedDataNumericEntryBox::OnCreateEditMenu)
-					.ButtonContent()
+					SNew(SHorizontalBox)
+					+ SHorizontalBox::Slot()
+					.HAlign(HAlign_Right)
+					.VAlign(VAlign_Center)
+					.FillWidth(1.f)
+					.Padding(4.0f, 0.0f, 0.0f, 0.0f)
 					[
-						SNew(STextBlock)
-						.Font(FAppStyle::Get().GetFontStyle("FontAwesome.8"))
-						.Text(FEditorFontGlyphs::Pencil_Square)
-						.ColorAndOpacity(FLinearColor::White)
+						TextBlock
 					]
-				]
-			];
+					+ SHorizontalBox::Slot()
+					.HAlign(HAlign_Right)
+					.VAlign(VAlign_Center)
+					.AutoWidth()
+					.Padding(4.0f, 0.0f, 0.0f, 0.0f)
+					[
+						SAssignNew(ComboButton, SComboButton)
+						.ComboButtonStyle(FTimedDataMonitorEditorStyle::Get(), "ToggleComboButton")
+						.HAlign(HAlign_Center)
+						.HasDownArrow(false)
+						.OnGetMenuContent(this, &STimedDataNumericEntryBox::OnCreateEditMenu)
+						.ButtonContent()
+						[
+							SNew(STextBlock)
+							.Font(FAppStyle::Get().GetFontStyle("FontAwesome.8"))
+							.Text(FEditorFontGlyphs::Pencil_Square)
+							.ColorAndOpacity(FLinearColor::White)
+						]
+					]
+				];
+			}
+			else
+			{
+				ChildSlot
+				[
+					TextBlock
+				];
+			}
 		}
 		else
 		{
-			ChildSlot
-			[
-				TextBlock
-			];
+
+			if (InArgs._CanEdit)
+			{
+				ChildSlot
+				[
+					OnCreateEditMenu()
+				];
+			}
+			else
+			{
+				ChildSlot
+				[
+					SNew(SHorizontalBox)
+					+ SHorizontalBox::Slot()
+					.HAlign(HAlign_Fill)
+					.VAlign(VAlign_Center)
+					.AutoWidth()
+					[
+						SNew(STextBlock)
+						.TextStyle(FTimedDataMonitorEditorStyle::Get(), "TextBlock.Regular")
+						.Text(InArgs._EditLabel)
+					]
+					+ SHorizontalBox::Slot()
+					.HAlign(HAlign_Fill)
+					.VAlign(VAlign_Center)
+					.AutoWidth()
+					[
+						TextBlock
+					]
+				];
+			}
 		}
 	}
 	
@@ -119,56 +165,79 @@ private:
 
 	TSharedRef<SWidget> OnCreateEditMenu()
 	{
-		bCloseRequested = false;
-
-		TSharedPtr<SWidget> WidgetToFocus;
-		TSharedRef<SWidget> Result = SNew(SBorder)
+		TSharedRef<SWidget> InnerWidget = SNew(SHorizontalBox)
+			+ SHorizontalBox::Slot()
 			.VAlign(VAlign_Center)
-			.Padding(FMargin(5.f, 5.f, 5.f, 5.f))
+			.HAlign(HAlign_Center)
+			.Padding(2.0f, 0.f)
+			.AutoWidth()
 			[
-				SNew(SHorizontalBox)
-				+ SHorizontalBox::Slot()
-				.HAlign(HAlign_Fill)
-				.VAlign(VAlign_Center)
-				.AutoWidth()
-				[
-					SNew(STextBlock)
-					.TextStyle(FTimedDataMonitorEditorStyle::Get(), "TextBlock.Regular")
-					.Text(EditLabel)
-				]
-				+ SHorizontalBox::Slot()
-				.HAlign(HAlign_Fill)
-				.VAlign(VAlign_Center)
-				.FillWidth(1.f)
-				.Padding(4.0f, 0.0f, 0.0f, 0.0f)
+				SNew(SButton)
+				.ButtonStyle(FTimedDataMonitorEditorStyle::Get(), "FlatButton")
+				.OnClicked(this, &STimedDataNumericEntryBox::OnMinusClicked)
+				.Content()
 				[
-					SAssignNew(WidgetToFocus, SNumericEntryBox<NumericType>)
-					.MinValue(TOptional<NumericType>(MinValue))
-					.MinDesiredValueWidth(50)
-					.Value(TOptional<NumericType>(Value.Get()))
-					.OnValueCommitted(this, &STimedDataNumericEntryBox::OnValueComittedCallback)
+					SNew(SImage)
+					.Image(FTimedDataMonitorEditorStyle::Get().GetBrush("MinusButton"))
 				]
-				+ SHorizontalBox::Slot()
-				.HAlign(HAlign_Right)
-				.VAlign(VAlign_Center)
-				.AutoWidth()
-				.Padding(4.0f, 0.0f, 0.0f, 0.0f)
+			]
+			+ SHorizontalBox::Slot()
+			.HAlign(HAlign_Fill)
+			.VAlign(VAlign_Center)
+			.FillWidth(1.f)
+			.Padding(2.0f, 0.0f, 0.0f, 0.0f)
+			[
+				SAssignNew(EntryBox, SNumericEntryBox<NumericType>)
+				.MinValue(MinValue)
+				.MaxValue(MaxValue)
+				.MinDesiredValueWidth(50)
+				.Value_Lambda([this]() 
+				{
+					CachedValue = Value.Get();
+					return TOptional<NumericType>(Value.Get()); 
+				})
+				.OnValueCommitted(this, &STimedDataNumericEntryBox::OnValueComittedCallback)
+
+			]
+			+ SHorizontalBox::Slot()
+			.VAlign(VAlign_Center)
+			.HAlign(HAlign_Center)
+			.Padding(2.0f, 0.0f, 0.0f, 0.0f)
+			.AutoWidth()
+			[
+				SNew(SButton)
+				.ButtonStyle(FTimedDataMonitorEditorStyle::Get(), "FlatButton")
+				.OnClicked(this, &STimedDataNumericEntryBox::OnPlusClicked)
+				.Content()
 				[
-					SNew(SCheckBox)
-					.Padding(4.f)
-					.Style(FAppStyle::Get(), "ToggleButtonCheckbox")
-					.IsChecked_Lambda([]() {return ECheckBoxState::Unchecked; })
-					.OnCheckStateChanged(this, &STimedDataNumericEntryBox::CloseComboButton)
-					[
-						SNew(STextBlock)
-						.Font(FAppStyle::Get().GetFontStyle("FontAwesome.11"))
-						.Text(FEditorFontGlyphs::Check)
-					]
+					SNew(SImage)
+					.Image(FTimedDataMonitorEditorStyle::Get().GetBrush("PlusButton"))
 				]
+			]
+			+ SHorizontalBox::Slot()
+			.HAlign(HAlign_Fill)
+			.VAlign(VAlign_Center)
+			.AutoWidth()
+			[
+				SuffixWidget
 			];
 
-		ComboButton->SetMenuContentWidgetToFocus(WidgetToFocus);
-		return Result;
+
+		if (bComboButton)
+		{
+			EditMenuContent = SNew(SBorder)
+				.VAlign(VAlign_Center)
+				.Padding(FMargin(5.f, 5.f, 5.f, 5.f))
+				[
+					InnerWidget
+				];
+		}
+		else
+		{
+			EditMenuContent = InnerWidget;
+		}
+		
+		return EditMenuContent.ToSharedRef();
 	}
 
 	void CloseComboButton(ECheckBoxState InNewState)
@@ -181,26 +250,62 @@ private:
 
 	void OnValueComittedCallback(NumericType NewValue, ETextCommit::Type CommitType)
 	{
-		if (!bCloseRequested)
+		if (MinValue.IsSet())
 		{
-			bCloseRequested = true;
+			NewValue = FMath::Max(MinValue.GetValue(), NewValue);
+		}
 
-			OnValueCommitted.ExecuteIfBound(NewValue, CommitType);
-			if (ComboButton)
+		if (MaxValue.IsSet())
+		{
+			NewValue = FMath::Min(MaxValue.GetValue(), NewValue);
+		}
+
+		CachedValue = NewValue;
+		OnValueCommitted.ExecuteIfBound(NewValue, CommitType);
+	}
+
+	FReply OnMinusClicked()
+	{
+		if (EntryBox)
+		{
+			CachedValue = Value.Get();
+			if (!MinValue.IsSet() || CachedValue - 1 >= MinValue.GetValue())
+			{
+				CachedValue--;
+			}
+			OnValueCommitted.ExecuteIfBound(CachedValue, ETextCommit::Type::Default);
+		}
+		return FReply::Handled();
+	}
+
+	FReply OnPlusClicked()
+	{
+		if (EntryBox)
+		{
+			CachedValue = Value.Get();
+			if (!MaxValue.IsSet() || CachedValue + 1 <= MinValue.GetValue())
 			{
-				ComboButton->SetIsOpen(false);
+				CachedValue++;
 			}
+			OnValueCommitted.ExecuteIfBound(CachedValue, ETextCommit::Type::Default);
 		}
+		return FReply::Handled();
 	}
 
 private:
 	TAttribute<NumericType> Value;
 	TAttribute<NumericType> Amount;
-	NumericType MinValue;
+	NumericType CachedValue;
+	TSharedPtr<SNumericEntryBox<NumericType>> EntryBox;
+	TOptional<NumericType> MinValue;
+	TOptional<NumericType> MaxValue;
 	FText EditLabel;
+	TSharedRef<SWidget> SuffixWidget = SNullWidget::NullWidget;
 	bool bShowAmount;
 	FOnValueCommitted OnValueCommitted;
 	TSharedPtr<SComboButton> ComboButton;
+	bool bComboButton = true;
 	bool bCloseRequested = false;
+	TSharedPtr<SWidget> EditMenuContent;
 };
 
diff --git a/Engine/Plugins/VirtualProduction/TimedDataMonitor/Source/TimedDataMonitorEditor/Private/STimedDataTimecodeProvider.cpp b/Engine/Plugins/VirtualProduction/TimedDataMonitor/Source/TimedDataMonitorEditor/Private/STimedDataTimecodeProvider.cpp
--- a/Engine/Plugins/VirtualProduction/TimedDataMonitor/Source/TimedDataMonitorEditor/Private/STimedDataTimecodeProvider.cpp
+++ b/Engine/Plugins/VirtualProduction/TimedDataMonitor/Source/TimedDataMonitorEditor/Private/STimedDataTimecodeProvider.cpp
@@ -20,6 +20,7 @@
 #include "Styling/CoreStyle.h"
 
 #include "STimedDataMonitorPanel.h"
+#include "STimedDataNumericEntryBox.h"
 #include "Widgets/Input/SCheckBox.h"
 #include "Widgets/Input/SSpinBox.h"
 #include "Widgets/Layout/SSeparator.h"
@@ -48,15 +49,14 @@ void STimedDataTimecodeProvider::Construct(const FArguments& InArgs, TSharedPtr<
 	UpdateCachedValue();
 
 	FSlateFontInfo TimeFont = FCoreStyle::Get().GetFontStyle(TEXT("EmbossedText"));
-	TimeFont.Size += 4;
+	TimeFont.Size += 2;
 
 	ChildSlot
 	[
 		SNew(SVerticalBox)
 		+ SVerticalBox::Slot()
-		.Padding(0)
+		.Padding(0.f)
 		.AutoHeight()
-		.HAlign(HAlign_Center)
 		[
 			SNew(SHorizontalBox)
 			+ SHorizontalBox::Slot()
@@ -78,6 +78,42 @@ void STimedDataTimecodeProvider::Construct(const FArguments& InArgs, TSharedPtr<
 				.Font(FCoreStyle::Get().GetFontStyle(TEXT("NormalText")))
 				.Text(this, &STimedDataTimecodeProvider::GetTimecodeProviderText)
 			]
+			
+			+ SHorizontalBox::Slot()
+			.AutoWidth()
+			[
+				SNew(SCheckBox)
+				.Padding(4.f)
+				.ToolTipText(LOCTEXT("ShowTimecodeProviderSetting_Tip", "Show timecode provider setting"))
+				.Style(FAppStyle::Get(), "ToggleButtonCheckbox")
+				.ForegroundColor(FSlateColor::UseForeground())
+				.IsEnabled(this, &STimedDataTimecodeProvider::IsTimecodeOffsetEnabled)
+				.IsChecked_Lambda([]() {return ECheckBoxState::Unchecked; })
+				.OnCheckStateChanged(this, &STimedDataTimecodeProvider::ShowTimecodeProviderSetting)
+				[
+					SNew(STextBlock)
+					.Font(FAppStyle::Get().GetFontStyle("FontAwesome.11"))
+					.Text(FEditorFontGlyphs::Cogs)
+				]
+			]
+			+ SHorizontalBox::Slot()
+			.AutoWidth()
+			.Padding(FMargin(4.f, 0.f, 0.f, 0.f))
+			[
+				SNew(SCheckBox)
+				.Padding(4.f)
+				.ToolTipText(LOCTEXT("ReapplyMenuToolTip", "Reinitialize the current Timecode Provider."))
+				.Style(FAppStyle::Get(), "ToggleButtonCheckbox")
+				.ForegroundColor(FSlateColor::UseForeground())
+				.IsEnabled(this, &STimedDataTimecodeProvider::IsTimecodeOffsetEnabled)
+				.IsChecked_Lambda([]() {return ECheckBoxState::Unchecked; })
+				.OnCheckStateChanged(this, &STimedDataTimecodeProvider::ReinitializeTimecodeProvider)
+				[
+					SNew(STextBlock)
+					.Font(FAppStyle::Get().GetFontStyle("FontAwesome.11"))
+					.Text(FEditorFontGlyphs::Undo)
+				]
+			]
 		]
 		+ SVerticalBox::Slot()
 		.Padding(0)
@@ -103,85 +139,31 @@ void STimedDataTimecodeProvider::Construct(const FArguments& InArgs, TSharedPtr<
 				//	.Text(this, &STimedDataTimecodeProvider::GetSystemTimeText)
 				//]
 			]
+		]
+		+ SVerticalBox::Slot()
+		.AutoHeight()
+		[
+			SNew(SHorizontalBox)
 			+ SHorizontalBox::Slot()
+			.VAlign(VAlign_Center)
 			.AutoWidth()
 			[
-				SNew(SSeparator)
-				.Thickness(2)
-				.Orientation(EOrientation::Orient_Vertical)
+				SNew(STextBlock)
+				.Font(FCoreStyle::Get().GetFontStyle(TEXT("NormalText")))
+				.Text(LOCTEXT("GlobalFrameOffsetLabel", "Global Frame Offset: "))
 			]
 			+ SHorizontalBox::Slot()
+			.VAlign(VAlign_Center)
 			.AutoWidth()
-			.Padding(FMargin(6.f, 4.f, 0.f, 4.f))
 			[
-				SNew(SVerticalBox)
-				+ SVerticalBox::Slot()
-				.VAlign(VAlign_Center)
-				[
-					SNew(SHorizontalBox)
-					+ SHorizontalBox::Slot()
-					.FillWidth(1)
-					.HAlign(HAlign_Right)
-					.VAlign(VAlign_Center)
-					[
-						SNew(STextBlock)
-						.Font(FCoreStyle::Get().GetFontStyle(TEXT("NormalText")))
-						.Text(LOCTEXT("FrameOffsetLabel", "Global TC Offset: "))
-					]
-					+ SHorizontalBox::Slot()
-					.AutoWidth()
-					.VAlign(VAlign_Center)
-					[
-						SNew(SSpinBox<float>)
-						.ToolTipText(LOCTEXT("FrameDelay_ToolTip", "Number of frames to subtract from the original timecode."))
-						.Value(this, &STimedDataTimecodeProvider::GetTimecodeFrameDelay)
-						.MinValue(TOptional<float>())
-						.MaxValue(TOptional<float>())
-						.OnValueCommitted(this, &STimedDataTimecodeProvider::SetTimecodeFrameDelay)
-						.IsEnabled(this, &STimedDataTimecodeProvider::IsTimecodeOffsetEnabled)
-					]
-				]
-				+ SVerticalBox::Slot()
-				.HAlign(HAlign_Right)
-				.VAlign(VAlign_Bottom)
-				[
-					SNew(SHorizontalBox)
-					+ SHorizontalBox::Slot()
-					.AutoWidth()
-					[
-						SNew(SCheckBox)
-						.Padding(4.f)
-						.ToolTipText(LOCTEXT("ShowTimecodeProviderSetting_Tip", "Show timecode provider setting"))
-						.Style(FAppStyle::Get(), "ToggleButtonCheckbox")
-						.ForegroundColor(FSlateColor::UseForeground())
-						.IsEnabled(this, &STimedDataTimecodeProvider::IsTimecodeOffsetEnabled)
-						.IsChecked_Lambda([]() {return ECheckBoxState::Unchecked; })
-						.OnCheckStateChanged(this, &STimedDataTimecodeProvider::ShowTimecodeProviderSetting)
-						[
-							SNew(STextBlock)
-							.Font(FAppStyle::Get().GetFontStyle("FontAwesome.11"))
-							.Text(FEditorFontGlyphs::Cogs)
-						]
-					]
-					+ SHorizontalBox::Slot()
-					.AutoWidth()
-					.Padding(FMargin(4.f, 0.f, 0.f, 0.f))
-					[
-						SNew(SCheckBox)
-						.Padding(4.f)
-						.ToolTipText(LOCTEXT("ReapplyMenuToolTip", "Reinitialize the current Timecode Provider."))
-						.Style(FAppStyle::Get(), "ToggleButtonCheckbox")
-						.ForegroundColor(FSlateColor::UseForeground())
-						.IsEnabled(this, &STimedDataTimecodeProvider::IsTimecodeOffsetEnabled)
-						.IsChecked_Lambda([]() {return ECheckBoxState::Unchecked; })
-						.OnCheckStateChanged(this, &STimedDataTimecodeProvider::ReinitializeTimecodeProvider)
-						[
-							SNew(STextBlock)
-							.Font(FAppStyle::Get().GetFontStyle("FontAwesome.11"))
-							.Text(FEditorFontGlyphs::Undo)
-						]
-					]
-				]
+				SNew(STimedDataNumericEntryBox<float>)
+				.ComboButton(false)
+				.ToolTipText(LOCTEXT("FrameDelay_ToolTip", "Number of frames to subtract from the original timecode."))
+				.Value(this, &STimedDataTimecodeProvider::GetTimecodeFrameDelay)
+				.MinValue(TOptional<float>())
+				.MaxValue(TOptional<float>())
+				.OnValueCommitted(this, &STimedDataTimecodeProvider::SetTimecodeFrameDelay)
+				.IsEnabled(this, &STimedDataTimecodeProvider::IsTimecodeOffsetEnabled)
 			]
 		]
 	];
diff --git a/Engine/Plugins/VirtualProduction/TimedDataMonitor/Source/TimedDataMonitorEditor/Private/STimingDiagramWidget.h b/Engine/Plugins/VirtualProduction/TimedDataMonitor/Source/TimedDataMonitorEditor/Private/STimingDiagramWidget.h
--- a/Engine/Plugins/VirtualProduction/TimedDataMonitor/Source/TimedDataMonitorEditor/Private/STimingDiagramWidget.h
+++ b/Engine/Plugins/VirtualProduction/TimedDataMonitor/Source/TimedDataMonitorEditor/Private/STimingDiagramWidget.h
@@ -24,6 +24,7 @@ public:
 			, _ShowMean(true)
 			, _ShowSigma(false)
 			, _ShowSnapshot(false)
+			, _UseNiceBrush(false)
 			, _SizePerSeconds(100)
 		{ }
 		SLATE_ARGUMENT(FTimedDataMonitorChannelIdentifier, ChannelIdentifier)
diff --git a/Engine/Plugins/VirtualProduction/TimedDataMonitor/Source/TimedDataMonitorEditor/Private/TimedDataMonitorEditorSettings.h b/Engine/Plugins/VirtualProduction/TimedDataMonitor/Source/TimedDataMonitorEditor/Private/TimedDataMonitorEditorSettings.h
--- a/Engine/Plugins/VirtualProduction/TimedDataMonitor/Source/TimedDataMonitorEditor/Private/TimedDataMonitorEditorSettings.h
+++ b/Engine/Plugins/VirtualProduction/TimedDataMonitor/Source/TimedDataMonitorEditor/Private/TimedDataMonitorEditorSettings.h
@@ -56,6 +56,18 @@ public:
 	UPROPERTY(Config, AdvancedDisplay, EditAnywhere, Category = "UI", meta = (ClampMin = 0, ClampMax = 5, EditCondition = "bOverrideNumberOfStandardDeviationToShow"))
 	int32 OverrideNumberOfStandardDeviationToShow = 5;
 
+	/**
+	 * The threshold at which to display the frame time as orange.
+	 */
+	UPROPERTY(Config, EditAnywhere, Category = "UI", meta = (ClampMin = 0, ClampMax = 2))
+	double FrameTimeWarnThreshold = 0.85;
+
+	/**
+	 * The percent at which to display the target frame time in the frame time indicator. 
+	 */
+	UPROPERTY(Config, EditAnywhere, Category = "UI", meta = (ClampMin = 0, ClampMax = 1))
+	double TargetFrameTimePercent = 0.8;
+
 	UPROPERTY(Config)
 	ETimedDataMonitorEditorCalibrationType LastCalibrationType = ETimedDataMonitorEditorCalibrationType::CalibrateWithTimecode;
 
diff --git a/Engine/Plugins/VirtualProduction/TimedDataMonitor/Source/TimedDataMonitorEditor/Private/TimedDataMonitorEditorStyle.cpp b/Engine/Plugins/VirtualProduction/TimedDataMonitor/Source/TimedDataMonitorEditor/Private/TimedDataMonitorEditorStyle.cpp
--- a/Engine/Plugins/VirtualProduction/TimedDataMonitor/Source/TimedDataMonitorEditor/Private/TimedDataMonitorEditorStyle.cpp
+++ b/Engine/Plugins/VirtualProduction/TimedDataMonitor/Source/TimedDataMonitorEditor/Private/TimedDataMonitorEditorStyle.cpp
@@ -25,9 +25,9 @@ FTimedDataMonitorEditorStyle::FTimedDataMonitorEditorStyle()
 
 	// CheckBox
 	{
-		FSlateImageBrush SwitchOn = FSlateImageBrush(RootToContentDir(TEXT("Widgets/Switch_ON.png")), FVector2D(28.f, 14.f));
-		FSlateImageBrush SwitchOff = FSlateImageBrush(RootToContentDir(TEXT("Widgets/Switch_OFF.png")), FVector2D(28.f, 14.f));
-		FSlateImageBrush SwitchUndeterminded = FSlateImageBrush(RootToContentDir(TEXT("Widgets/Switch_Undetermined.png")), FVector2D(28.f, 14.f));
+		FSlateImageBrush SwitchOn = IMAGE_BRUSH(TEXT("Widgets/Switch_ON"), FVector2D(28.f, 14.f));
+		FSlateImageBrush SwitchOff = IMAGE_BRUSH(TEXT("Widgets/Switch_OFF"), FVector2D(28.f, 14.f));
+		FSlateImageBrush SwitchUndeterminded = IMAGE_BRUSH(TEXT("Widgets/Switch_Undetermined"), FVector2D(28.f, 14.f));
 
 		FCheckBoxStyle SwitchStyle = FCheckBoxStyle()
 			.SetForegroundColor(FLinearColor::White)
@@ -54,14 +54,14 @@ FTimedDataMonitorEditorStyle::FTimedDataMonitorEditorStyle()
 		const FVector2D Icon20x20(20.0f, 20.0f);
 		Set("Img.TimedDataMonitor.Icon", new IMAGE_BRUSH_SVG("Common/TimedDataMonitor", Icon20x20));
 
-		Set(NAME_TimecodeBrush, new FSlateImageBrush(RootToContentDir(TEXT("Widgets/Timecode_16x.png")), Icon16x16));
-		Set(NAME_PlatformTimeBrush, new FSlateImageBrush(RootToContentDir(TEXT("Widgets/Time_16x.png")), Icon16x16));
-		Set(NAME_NoEvaluationBrush, new FSlateImageBrush(RootToContentDir(TEXT("Widgets/NoEvaluation_16x.png")), Icon16x16));
+		Set(NAME_TimecodeBrush, new IMAGE_BRUSH(TEXT("Widgets/Timecode_16x"), Icon16x16));
+		Set(NAME_PlatformTimeBrush, new IMAGE_BRUSH(TEXT("Widgets/Time_16x"), Icon16x16));
+		Set(NAME_NoEvaluationBrush, new IMAGE_BRUSH(TEXT("Widgets/NoEvaluation_16x"), Icon16x16));
 
-		Set("Img.BufferVisualization", new FSlateImageBrush(RootToContentDir(TEXT("Widgets/BufferVisualization_24x.png")), Icon24x24));
-		Set("Img.Calibration", new FSlateImageBrush(RootToContentDir(TEXT("Widgets/Calibration_24x.png")), Icon24x24));
-		Set("Img.TimeCorrection", new FSlateImageBrush(RootToContentDir(TEXT("Widgets/TimeCorrection_24x.png")), Icon24x24));
-		Set("Img.Edit", new FSlateImageBrush(RootToContentDir(TEXT("Widgets/Edit_24x.png")), Icon24x24));
+		Set("Img.BufferVisualization", new IMAGE_BRUSH(TEXT("Widgets/BufferVisualization_24x"), Icon24x24));
+		Set("Img.Calibration", new IMAGE_BRUSH(TEXT("Widgets/Calibration_24x"), Icon24x24));
+		Set("Img.TimeCorrection", new IMAGE_BRUSH(TEXT("Widgets/TimeCorrection_24x"), Icon24x24));
+		Set("Img.Edit", new IMAGE_BRUSH(TEXT("Widgets/Edit_24x"), Icon24x24));
 		
 	}
 
@@ -127,10 +127,10 @@ FTimedDataMonitorEditorStyle::FTimedDataMonitorEditorStyle()
 	// Button
 	{
 		FButtonStyle FlatButton = FButtonStyle()
-			.SetNormal(FSlateBoxBrush(RootToContentDir("Common/ButtonHoverHint.png"), FMargin(4 / 16.0f), FLinearColor(1, 1, 1, 0.15f)))
-			.SetHovered(FSlateBoxBrush(RootToContentDir("Common/ButtonHoverHint.png"), FMargin(4 / 16.0f), FLinearColor(1, 1, 1, 0.25f)))
-			.SetPressed(FSlateBoxBrush(RootToContentDir("Common/ButtonHoverHint.png"), FMargin(4 / 16.0f), FLinearColor(1, 1, 1, 0.30f)))
-			.SetNormalPadding(FMargin(0, 0, 0, 1))
+			.SetNormal(BOX_BRUSH(TEXT("Common/ButtonHoverHint"), FMargin(4 / 16.0f), FLinearColor(1, 1, 1, 0.15f)))
+			.SetHovered(BOX_BRUSH(TEXT("Common/ButtonHoverHint"), FMargin(4 / 16.0f), FLinearColor(1, 1, 1, 0.25f)))
+			.SetPressed(BOX_BRUSH(TEXT("Common/ButtonHoverHint"), FMargin(4 / 16.0f), FLinearColor(1, 1, 1, 0.30f)))
+			.SetNormalPadding(FMargin(0.f, 2.f))
 			.SetPressedPadding(FMargin(0, 1, 0, 0));
 		Set("FlatButton", FlatButton);
 
@@ -139,9 +139,14 @@ FTimedDataMonitorEditorStyle::FTimedDataMonitorEditorStyle()
 			.SetNormal(ToggleButtonStyle.UncheckedImage)
 			.SetHovered(ToggleButtonStyle.UncheckedHoveredImage)
 			.SetPressed(ToggleButtonStyle.UncheckedPressedImage)
-			.SetNormalPadding(FMargin(0, 0, 0, 0))
+			.SetNormalPadding(FMargin(0, 0, 0, 1))
 			.SetPressedPadding(FMargin(0, 1, 0, 0));
 		Set("ToggleButton", ToggleButton);
+
+
+		const FVector2D Icon12x12(12.0f, 12.0f);
+		Set("PlusButton", new FSlateImageBrush( FAppStyle::Get().GetContentRootDir() / (TEXT("Icons/PlusSymbol_12x.png")), Icon12x12, FLinearColor::Gray));
+		Set("MinusButton", new FSlateImageBrush( FAppStyle::Get().GetContentRootDir() / TEXT("Icons/MinusSymbol_12x.png"), Icon12x12,  FLinearColor::Gray));
 	}
 	
 	FSlateStyleRegistry::RegisterSlateStyle(*this);
diff --git a/Engine/Plugins/VirtualProduction/TimedDataMonitor/Source/TimedDataMonitorEditor/TimedDataMonitorEditor.build.cs b/Engine/Plugins/VirtualProduction/TimedDataMonitor/Source/TimedDataMonitorEditor/TimedDataMonitorEditor.build.cs
--- a/Engine/Plugins/VirtualProduction/TimedDataMonitor/Source/TimedDataMonitorEditor/TimedDataMonitorEditor.build.cs
+++ b/Engine/Plugins/VirtualProduction/TimedDataMonitor/Source/TimedDataMonitorEditor/TimedDataMonitorEditor.build.cs
@@ -11,6 +11,7 @@ public class TimedDataMonitorEditor : ModuleRules
 				"Core",
 				"CoreUObject",
 				"EditorFramework",
+				"EditorWidgets",
 				"Engine",
 				"EditorStyle",
 				"InputCore",
diff --git a/Engine/Source/Runtime/TimeManagement/Private/ITimedDataInput.cpp b/Engine/Source/Runtime/TimeManagement/Private/ITimedDataInput.cpp
--- a/Engine/Source/Runtime/TimeManagement/Private/ITimedDataInput.cpp
+++ b/Engine/Source/Runtime/TimeManagement/Private/ITimedDataInput.cpp
@@ -15,3 +15,15 @@ double ITimedDataInput::ConvertFrameOffsetInSecondOffset(double Frames, FFrameRa
 {
 	return Rate.AsSeconds(FFrameTime::FromDecimal(Frames));
 }
+
+double ITimedDataInput::ConvertSecondOffsetInFrameOffset(double Seconds) const
+{
+	const FFrameRate FrameRate = GetFrameRate();
+	return ITimedDataInput::ConvertSecondOffsetInFrameOffset(Seconds, FrameRate);
+}
+
+double ITimedDataInput::ConvertFrameOffsetInSecondOffset(double Frames) const
+{
+	const FFrameRate FrameRate = GetFrameRate();
+	return ITimedDataInput::ConvertFrameOffsetInSecondOffset(Frames, FrameRate);
+}
\ No newline at end of file
diff --git a/Engine/Source/Runtime/TimeManagement/Public/ITimedDataInput.h b/Engine/Source/Runtime/TimeManagement/Public/ITimedDataInput.h
--- a/Engine/Source/Runtime/TimeManagement/Public/ITimedDataInput.h
+++ b/Engine/Source/Runtime/TimeManagement/Public/ITimedDataInput.h
@@ -31,7 +31,6 @@ enum class ETimedDataInputEvaluationType : uint8
 	PlatformTime,
 };
 
-
 UENUM()
 enum class ETimedDataInputState : uint8
 {
@@ -111,7 +110,7 @@ public:
 	/** Set the offset in seconds used at evaluation. */
 	virtual void SetEvaluationOffsetInSeconds(double Offset) = 0;
 
-	/** Get the frame rate at which the samples is produce. */
+	/** Get the frame rate at which the samples are produced. */
 	virtual FFrameRate GetFrameRate() const = 0;
 
 	/** Does channel from this input support a different buffer size than it's input. */
@@ -129,6 +128,12 @@ public:
 	/** Remove channel from the input */
 	virtual void RemoveChannel(ITimedDataInputChannel* Channel) = 0;
 
+	/** Convert second offset to frame offset using this input's framerate. */
+	double ConvertSecondOffsetInFrameOffset(double Seconds) const;
+
+	/** Convert frame offset to second offset using this input's framerate. */
+	double ConvertFrameOffsetInSecondOffset(double Frames) const;
+
 #if WITH_EDITOR
 	/** Get the icon that represent the input. */
 	virtual const FSlateBrush* GetDisplayIcon() const = 0;
