Add modifier key to increment/decrement timed data monitor offset by 10 and 0.1
By default, Shift key will increment/decrement by 0.1, Control key by 10
Add subframe restriction for inputs that do not support subframes

diff --git a/Engine/Plugins/Media/MediaIOFramework/Source/MediaIOCore/Private/Player/MediaIOCorePlayerBase.cpp b/Engine/Plugins/Media/MediaIOFramework/Source/MediaIOCore/Private/Player/MediaIOCorePlayerBase.cpp
--- a/Engine/Plugins/Media/MediaIOFramework/Source/MediaIOCore/Private/Player/MediaIOCorePlayerBase.cpp
+++ b/Engine/Plugins/Media/MediaIOFramework/Source/MediaIOCore/Private/Player/MediaIOCorePlayerBase.cpp
@@ -856,4 +856,9 @@ void FMediaIOCorePlayerBase::RemoveChannel(ITimedDataInputChannel* Channel)
 	Channels.Remove(Channel);
 }
 
+bool FMediaIOCorePlayerBase::SupportsSubFrames() const
+{
+	return false;
+}
+
 #undef LOCTEXT_NAMESPACE
diff --git a/Engine/Plugins/Media/MediaIOFramework/Source/MediaIOCore/Public/MediaIOCorePlayerBase.h b/Engine/Plugins/Media/MediaIOFramework/Source/MediaIOCore/Public/MediaIOCorePlayerBase.h
--- a/Engine/Plugins/Media/MediaIOFramework/Source/MediaIOCore/Public/MediaIOCorePlayerBase.h
+++ b/Engine/Plugins/Media/MediaIOFramework/Source/MediaIOCore/Public/MediaIOCorePlayerBase.h
@@ -147,7 +147,8 @@ public:
 	virtual FFrameRate GetFrameRate() const override;
 	virtual bool IsDataBufferSizeControlledByInput() const override;
 	virtual void AddChannel(ITimedDataInputChannel* Channel) override;
-	virtual void RemoveChannel(ITimedDataInputChannel * Channel) override; 
+	virtual void RemoveChannel(ITimedDataInputChannel * Channel) override;
+	virtual bool SupportsSubFrames() const override;
 
 protected:
 	/** Is the IO hardware/device ready to be used. */
diff --git a/Engine/Plugins/VirtualProduction/TimedDataMonitor/Source/TimedDataMonitor/Private/TimedDataMonitorSubsystem.cpp b/Engine/Plugins/VirtualProduction/TimedDataMonitor/Source/TimedDataMonitor/Private/TimedDataMonitorSubsystem.cpp
--- a/Engine/Plugins/VirtualProduction/TimedDataMonitor/Source/TimedDataMonitor/Private/TimedDataMonitorSubsystem.cpp
+++ b/Engine/Plugins/VirtualProduction/TimedDataMonitor/Source/TimedDataMonitor/Private/TimedDataMonitorSubsystem.cpp
@@ -474,8 +474,16 @@ void UTimedDataMonitorSubsystem::SetInputEvaluationOffsetInSeconds(const FTimedD
 {
 	BuildSourcesListIfNeeded();
 
-	if (FTimeDataInputItem* SourceItem = InputMap.Find(Identifier))
+	if (const FTimeDataInputItem* SourceItem = InputMap.Find(Identifier))
 	{
+		if (!SourceItem->Input->SupportsSubFrames())
+		{
+			// If sub frames are not supported, convert to frames, round it, then re-convert to seconds in order to get
+			// a time that aligns with a frame.
+			const int32 RoundedFrameOffset = FMath::RoundToInt32(SourceItem->Input->ConvertSecondOffsetInFrameOffset(Offset));
+			const double AlignedSecondsOffset = SourceItem->Input->ConvertFrameOffsetInSecondOffset(RoundedFrameOffset);
+			return SourceItem->Input->SetEvaluationOffsetInSeconds(AlignedSecondsOffset);
+		}
 		return SourceItem->Input->SetEvaluationOffsetInSeconds(Offset);
 	}
 }
diff --git a/Engine/Plugins/VirtualProduction/TimedDataMonitor/Source/TimedDataMonitorEditor/Private/STimedDataListView.cpp b/Engine/Plugins/VirtualProduction/TimedDataMonitor/Source/TimedDataMonitorEditor/Private/STimedDataListView.cpp
--- a/Engine/Plugins/VirtualProduction/TimedDataMonitor/Source/TimedDataMonitorEditor/Private/STimedDataListView.cpp
+++ b/Engine/Plugins/VirtualProduction/TimedDataMonitor/Source/TimedDataMonitorEditor/Private/STimedDataListView.cpp
@@ -21,6 +21,7 @@
 #include "Widgets/Images/SImage.h"
 #include "Widgets/Input/SCheckBox.h"
 #include "Widgets/Input/SNumericEntryBox.h"
+#include "Widgets/Layout/SWidgetSwitcher.h"
 #include "Widgets/SNullWidget.h"
 #include "Widgets/Text/STextBlock.h"
 #include "Widgets/Views/SExpanderArrow.h"
@@ -65,7 +66,7 @@ struct FTimedDataInputTableRowData : TSharedFromThis<FTimedDataInputTableRowData
 
 		DisplayName = TimedDataMonitorSubsystem->GetInputDisplayName(InputIdentifier);
 
-		ETimedDataInputEvaluationType InputType = TimedDataMonitorSubsystem->GetInputEvaluationType(InputIdentifier);
+		const ETimedDataInputEvaluationType InputType = TimedDataMonitorSubsystem->GetInputEvaluationType(InputIdentifier);
 		if (InputType == ETimedDataInputEvaluationType::Timecode)
 		{
 			CachedInputEvaluationOffsetType = ETimedDataInputEvaluationOffsetType::Frames;
@@ -75,9 +76,10 @@ struct FTimedDataInputTableRowData : TSharedFromThis<FTimedDataInputTableRowData
 			CachedInputEvaluationOffsetType = ETimedDataInputEvaluationOffsetType::Seconds;
 		}
 
-		if (ITimedDataInput* InputTimedData = TimedDataMonitorSubsystem->GetTimedDataInput(InputIdentifier))
+		if (const ITimedDataInput* InputTimedData = TimedDataMonitorSubsystem->GetTimedDataInput(InputIdentifier))
 		{
 			InputIcon = InputTimedData->GetDisplayIcon();
+			bSupportsSubFrames = InputTimedData->SupportsSubFrames();
 		}
 	}
 
@@ -124,6 +126,10 @@ struct FTimedDataInputTableRowData : TSharedFromThis<FTimedDataInputTableRowData
 			bCachedCanEditBufferSize = (CachedEnabled == ECheckBoxState::Checked || CachedEnabled == ECheckBoxState::Undetermined) && bControlBufferSize;
 
 			NewestDataTime = TimedDataMonitorSubsystem->GetInputNewestDataTime(InputIdentifier);
+			if (NewestDataTime.Timecode.Time == 0)
+			{
+				NewestDataTime.Timecode.Rate = TimedDataMonitorSubsystem->GetInputFrameRate(InputIdentifier);
+			}
 
 			CachedStatsBufferUnderflow = 0;
 			CachedStatsBufferOverflow = 0;
@@ -163,8 +169,16 @@ struct FTimedDataInputTableRowData : TSharedFromThis<FTimedDataInputTableRowData
 			{
 			case ETimedDataInputEvaluationType::Timecode:
 			{
-				FTimecode Timecode = FTimecode::FromFrameNumber(NewestDataTime.Timecode.Time.GetFrame(), NewestDataTime.Timecode.Rate);
-				CachedDescription = FText::Format(LOCTEXT("TimecodeDescription", "{0}@{1}"), FText::FromString(Timecode.ToString()), NewestDataTime.Timecode.Rate.ToPrettyText());
+				const FTimecode Timecode = FTimecode::FromFrameNumber(NewestDataTime.Timecode.Time.GetFrame(), NewestDataTime.Timecode.Rate);
+				if (Timecode == FTimecode())
+				{
+					// Timecode isn't valid.
+					CachedDescription = LOCTEXT("NoTimecodeLabel", "No Timecode");
+				}
+				else
+				{
+					CachedDescription = FText::Format(INVTEXT("{0}@{1}"), FText::FromString(Timecode.ToString()), NewestDataTime.Timecode.Rate.ToPrettyText());
+				}
 			}
 			break;
 			case ETimedDataInputEvaluationType::PlatformTime:
@@ -189,6 +203,7 @@ public:
 	FText DisplayName;
 	const FSlateBrush* InputIcon = nullptr;
 	TArray<FTimedDataInputTableRowDataPtr> InputChildren;
+	bool bSupportsSubFrames = true;
 
 	ECheckBoxState CachedEnabled = ECheckBoxState::Undetermined;
 	ETimedDataInputEvaluationType CachedInputEvaluationType = ETimedDataInputEvaluationType::None;
@@ -334,18 +349,41 @@ TSharedRef<SWidget> STimedDataInputTableRow::GenerateWidgetForColumn(const FName
 	{
 		if (Item->bIsInput)
 		{
-			TSharedRef<SWidget> Label = SNew(SEnumComboBox, StaticEnum<ETimedDataInputEvaluationOffsetType>())
+			const TSharedRef<SWidget> Label = SNew(SEnumComboBox, StaticEnum<ETimedDataInputEvaluationOffsetType>())
 				.ContentPadding(0.f)
 				.CurrentValue_Lambda([this]() { return static_cast<int32>(Item->CachedInputEvaluationOffsetType); })
 				.OnEnumSelectionChanged(this, &STimedDataInputTableRow::OnEvaluationOffsetTypeChanged);
 
-			return SNew(STimedDataNumericEntryBox<float>)
-				.TextStyle(ItemTextBlockStyle)
-				.ToolTipText(LOCTEXT("TimeCorrection_ToolTip", "Time Correction."))
-				.Value(this, &STimedDataInputTableRow::GetEvaluationOffset)
-				.ComboButton(false)
-				.SuffixWidget(Label)
-				.OnValueCommitted(this, &STimedDataInputTableRow::SetEvaluationOffset);
+			return SNew(SWidgetSwitcher)
+				.WidgetIndex_Lambda([this]()
+				{
+					bool bUseFloatEntryBox = true;
+					if (!Item->bSupportsSubFrames)
+					{
+						bUseFloatEntryBox = Item->CachedInputEvaluationOffsetType == ETimedDataInputEvaluationOffsetType::Seconds;
+					}
+					return bUseFloatEntryBox ? 0 : 1;
+				})
+				+SWidgetSwitcher::Slot()
+				[
+					SNew(STimedDataNumericEntryBox<double>)
+						.TextStyle(ItemTextBlockStyle)
+						.ToolTipText(LOCTEXT("TimeCorrection_ToolTip", "Time Correction."))
+						.Value(this, &STimedDataInputTableRow::GetFloatEvaluationOffset)
+						.ComboButton(false)
+						.SuffixWidget(Label)
+						.OnValueCommitted(this, &STimedDataInputTableRow::SetEvaluationOffset)
+				]
+				+SWidgetSwitcher::Slot()
+				[
+					SNew(STimedDataNumericEntryBox<int32>)
+						.TextStyle(ItemTextBlockStyle)
+						.ToolTipText(LOCTEXT("TimeCorrection_ToolTip", "Time Correction."))
+						.Value(this, &STimedDataInputTableRow::GetIntEvaluationOffset)
+						.ComboButton(false)
+						.SuffixWidget(Label)
+						.OnValueCommitted(this, &STimedDataInputTableRow::SetEvaluationOffset)
+				];
 		}
 		return SNullWidget::NullWidget;
 	}
@@ -493,15 +531,22 @@ FText STimedDataInputTableRow::GetEvaluationOffsetText() const
 }
 
 
-float STimedDataInputTableRow::GetEvaluationOffset() const
+double STimedDataInputTableRow::GetFloatEvaluationOffset() const
 {
-	return Item->CachedInputEvaluationOffsetType == ETimedDataInputEvaluationOffsetType::Seconds ? Item->CachedInputEvaluationOffsetSeconds : Item->CachedInputEvaluationOffsetFrames;
+	const double Result = Item->CachedInputEvaluationOffsetType == ETimedDataInputEvaluationOffsetType::Seconds ? Item->CachedInputEvaluationOffsetSeconds : Item->CachedInputEvaluationOffsetFrames;
+	return FMath::RoundToFloat(Result * 10000.0) / 10000.0;
 }
 
 
-void STimedDataInputTableRow::SetEvaluationOffset(float InValue, ETextCommit::Type CommitType)
+int32 STimedDataInputTableRow::GetIntEvaluationOffset() const
 {
-	if (GetEvaluationOffset() != InValue && Item->bIsInput)
+	return FMath::RoundToInt32(Item->CachedInputEvaluationOffsetType == ETimedDataInputEvaluationOffsetType::Seconds ? Item->CachedInputEvaluationOffsetSeconds : Item->CachedInputEvaluationOffsetFrames);
+}
+
+
+void STimedDataInputTableRow::SetEvaluationOffset(double InValue, ETextCommit::Type CommitType)
+{
+	if (GetFloatEvaluationOffset() != InValue && Item->bIsInput)
 	{
 		UTimedDataMonitorSubsystem* TimedDataMonitorSubsystem = GEngine->GetEngineSubsystem<UTimedDataMonitorSubsystem>();
 		check(TimedDataMonitorSubsystem);
@@ -526,6 +571,11 @@ void STimedDataInputTableRow::SetEvaluationOffset(float InValue, ETextCommit::Ty
 }
 
 
+void STimedDataInputTableRow::SetEvaluationOffset(int32 InValue, ETextCommit::Type CommitType)
+{
+	SetEvaluationOffset(static_cast<float>(InValue), CommitType);
+}
+
 int32 STimedDataInputTableRow::GetBufferSize() const
 {
 	return Item->CachedBufferSize;
diff --git a/Engine/Plugins/VirtualProduction/TimedDataMonitor/Source/TimedDataMonitorEditor/Private/STimedDataListView.h b/Engine/Plugins/VirtualProduction/TimedDataMonitor/Source/TimedDataMonitorEditor/Private/STimedDataListView.h
--- a/Engine/Plugins/VirtualProduction/TimedDataMonitor/Source/TimedDataMonitorEditor/Private/STimedDataListView.h
+++ b/Engine/Plugins/VirtualProduction/TimedDataMonitor/Source/TimedDataMonitorEditor/Private/STimedDataListView.h
@@ -66,9 +66,11 @@ private:
 	FSlateColor GetStateColorAndOpacity() const;
 
 	FText GetDescription() const;
-	float GetEvaluationOffset() const;
+	double GetFloatEvaluationOffset() const;
+	int32 GetIntEvaluationOffset() const;
 	FText GetEvaluationOffsetText() const;
-	void SetEvaluationOffset(float NewValue, ETextCommit::Type CommitType);
+	void SetEvaluationOffset(double NewValue, ETextCommit::Type CommitType);
+	void SetEvaluationOffset(int32 NewValue, ETextCommit::Type CommitType);
 	int32 GetBufferSize() const;
 	FText GetBufferSizeText() const;
 	void SetBufferSize(int32 NewValue, ETextCommit::Type CommitType);
diff --git a/Engine/Plugins/VirtualProduction/TimedDataMonitor/Source/TimedDataMonitorEditor/Private/STimedDataNumericEntryBox.h b/Engine/Plugins/VirtualProduction/TimedDataMonitor/Source/TimedDataMonitorEditor/Private/STimedDataNumericEntryBox.h
--- a/Engine/Plugins/VirtualProduction/TimedDataMonitor/Source/TimedDataMonitorEditor/Private/STimedDataNumericEntryBox.h
+++ b/Engine/Plugins/VirtualProduction/TimedDataMonitor/Source/TimedDataMonitorEditor/Private/STimedDataNumericEntryBox.h
@@ -4,20 +4,17 @@
 
 #include "CoreMinimal.h"
 
-#include "SlateFwd.h"
-#include "Widgets/DeclarativeSyntaxSupport.h"
-#include "Widgets/SCompoundWidget.h"
-#include "SlateOptMacros.h"
-
+#include "Editor.h"
 #include "EditorFontGlyphs.h"
+#include "Framework/Application/SlateApplication.h"
 #include "Styling/AppStyle.h"
 #include "TimedDataMonitorEditorStyle.h"
-
+#include "Widgets/DeclarativeSyntaxSupport.h"
 #include "Widgets/Images/SImage.h"
-#include "Widgets/Input/SCheckBox.h"
 #include "Widgets/Input/SButton.h"
 #include "Widgets/Input/SNumericEntryBox.h"
 #include "Widgets/Layout/SBorder.h"
+#include "Widgets/SCompoundWidget.h"
 #include "Widgets/SNullWidget.h"
 #include "Widgets/Text/STextBlock.h"
 
@@ -39,7 +36,20 @@ public:
 			, _ShowAmount(false)
 			, _TextStyle(&FAppStyle::Get().GetWidgetStyle<FTextBlockStyle>(TEXT("NormalText")))
 			, _ComboButton(true)
-		{ }
+			, _ControlModifierOffset(10)
+	{
+		_ShiftModifierOffset = NumericType(1) / 10;
+		if (!_ShiftModifierOffset)
+		{
+			_ShiftModifierOffset = 1;
+		}
+
+		_AltModifierOffset = NumericType(1) / 100;
+		if (!_AltModifierOffset)
+		{
+			_AltModifierOffset = 1;
+		}
+	}
 		SLATE_ATTRIBUTE(NumericType, Value)
 		SLATE_ARGUMENT(TOptional<NumericType>, MinValue)
 		SLATE_ARGUMENT(TOptional<NumericType>, MaxValue)
@@ -51,6 +61,9 @@ public:
 		SLATE_STYLE_ARGUMENT(FTextBlockStyle, TextStyle)
 		SLATE_ARGUMENT(bool, ComboButton)
 		SLATE_EVENT(FOnValueCommitted, OnValueCommitted)
+		SLATE_ARGUMENT(NumericType, ControlModifierOffset)
+		SLATE_ARGUMENT(NumericType, AltModifierOffset)
+		SLATE_ARGUMENT(NumericType, ShiftModifierOffset)
 	SLATE_END_ARGS()
 	
 public:
@@ -65,6 +78,9 @@ public:
 		bShowAmount = InArgs._ShowAmount;
 		bComboButton = InArgs._ComboButton;
 		OnValueCommitted = InArgs._OnValueCommitted;
+		AltOffset = InArgs._AltModifierOffset;
+		ControlOffset = InArgs._ControlModifierOffset;
+		ShiftOffset = InArgs._ShiftModifierOffset;
 
 		TSharedRef<SWidget> TextBlock = SNew(STextBlock)
 			.TextStyle(InArgs._TextStyle)
@@ -149,7 +165,7 @@ public:
 			}
 		}
 	}
-	
+
 private:
 	FText GetValueText() const 
 	{
@@ -196,7 +212,7 @@ private:
 					CachedValue = Value.Get();
 					return TOptional<NumericType>(Value.Get()); 
 				})
-				.OnValueCommitted(this, &STimedDataNumericEntryBox::OnValueComittedCallback)
+				.OnValueCommitted(this, &STimedDataNumericEntryBox::OnValueCommittedCallback)
 
 			]
 			+ SHorizontalBox::Slot()
@@ -248,7 +264,7 @@ private:
 		}
 	}
 
-	void OnValueComittedCallback(NumericType NewValue, ETextCommit::Type CommitType)
+	void OnValueCommittedCallback(NumericType NewValue, ETextCommit::Type CommitType)
 	{
 		if (MinValue.IsSet())
 		{
@@ -262,17 +278,33 @@ private:
 
 		CachedValue = NewValue;
 		OnValueCommitted.ExecuteIfBound(NewValue, CommitType);
+		if (!ComboButton && CommitType == ETextCommit::Type::OnEnter)
+		{
+			// Clear focus to get the updated value in case it was modified externally.
+			FSlateApplication::Get().ClearKeyboardFocus();
+			FTimerHandle Handle;
+			constexpr float InRate = 0.01f;
+			constexpr bool InbLoop = false;
+			GEditor->GetTimerManager()->SetTimer(Handle, [this](){ FSlateApplication::Get().SetKeyboardFocus(EntryBox); }, InRate, InbLoop);
+		}
 	}
 
 	FReply OnMinusClicked()
 	{
+		NumericType Offset = GetModifierKeyOffset();
 		if (EntryBox)
 		{
 			CachedValue = Value.Get();
-			if (!MinValue.IsSet() || CachedValue - 1 >= MinValue.GetValue())
+			if (!MinValue.IsSet())
 			{
-				CachedValue--;
+				CachedValue -= Offset;
 			}
+
+			if (MinValue.IsSet() && CachedValue - Offset < MinValue.GetValue())
+			{
+				CachedValue = MinValue.GetValue();
+			}
+			
 			OnValueCommitted.ExecuteIfBound(CachedValue, ETextCommit::Type::Default);
 		}
 		return FReply::Handled();
@@ -280,18 +312,48 @@ private:
 
 	FReply OnPlusClicked()
 	{
+		NumericType Offset = GetModifierKeyOffset();
 		if (EntryBox)
 		{
 			CachedValue = Value.Get();
-			if (!MaxValue.IsSet() || CachedValue + 1 <= MinValue.GetValue())
+			if (!MaxValue.IsSet())
+			{
+				CachedValue += Offset;
+			}
+
+			if (MaxValue.IsSet() && CachedValue + Offset > MaxValue.GetValue())
 			{
-				CachedValue++;
+				CachedValue = MaxValue.GetValue();
 			}
 			OnValueCommitted.ExecuteIfBound(CachedValue, ETextCommit::Type::Default);
 		}
 		return FReply::Handled();
 	}
 
+	NumericType GetModifierKeyOffset() const
+	{
+		const bool bIsShiftDown = FSlateApplication::Get().GetModifierKeys().IsShiftDown();
+		const bool bIsControlDown = FSlateApplication::Get().GetModifierKeys().IsControlDown();
+		const bool bIsAltDown = FSlateApplication::Get().GetModifierKeys().IsAltDown();
+		
+		NumericType Offset = 1;
+		
+		if (bIsShiftDown)
+		{
+			Offset = ShiftOffset;
+		}
+		else if (bIsControlDown)
+		{
+			Offset = ControlOffset;
+		}
+		else if (bIsAltDown)
+		{
+			Offset = AltOffset;
+		}
+
+		return Offset;
+	}
+
 private:
 	TAttribute<NumericType> Value;
 	TAttribute<NumericType> Amount;
@@ -307,5 +369,8 @@ private:
 	bool bComboButton = true;
 	bool bCloseRequested = false;
 	TSharedPtr<SWidget> EditMenuContent;
+	NumericType AltOffset;
+	NumericType ShiftOffset;
+	NumericType ControlOffset;
 };
 
diff --git a/Engine/Source/Runtime/TimeManagement/Public/ITimedDataInput.h b/Engine/Source/Runtime/TimeManagement/Public/ITimedDataInput.h
--- a/Engine/Source/Runtime/TimeManagement/Public/ITimedDataInput.h
+++ b/Engine/Source/Runtime/TimeManagement/Public/ITimedDataInput.h
@@ -92,6 +92,8 @@ public:
 	static double ConvertFrameOffsetInSecondOffset(double Frames, FFrameRate Rate);
 	
 public:
+	virtual ~ITimedDataInput() {}
+	
 	/** Get the name used when displayed. */
 	virtual FText GetDisplayName() const = 0;
 
@@ -128,6 +130,12 @@ public:
 	/** Remove channel from the input */
 	virtual void RemoveChannel(ITimedDataInputChannel* Channel) = 0;
 
+	/** Whether this input supports sub frames. */
+	virtual bool SupportsSubFrames() const
+	{
+		return true;
+	}
+
 	/** Convert second offset to frame offset using this input's framerate. */
 	double ConvertSecondOffsetInFrameOffset(double Seconds) const;
 
