Fix buttons in stage data monitor

diff --git a/Engine/Plugins/VirtualProduction/TimedDataMonitor/Source/TimedDataMonitorEditor/Private/STimedDataNumericEntryBox.h b/Engine/Plugins/VirtualProduction/TimedDataMonitor/Source/TimedDataMonitorEditor/Private/STimedDataNumericEntryBox.h
--- a/Engine/Plugins/VirtualProduction/TimedDataMonitor/Source/TimedDataMonitorEditor/Private/STimedDataNumericEntryBox.h
+++ b/Engine/Plugins/VirtualProduction/TimedDataMonitor/Source/TimedDataMonitorEditor/Private/STimedDataNumericEntryBox.h
@@ -295,10 +295,7 @@ private:
 		if (EntryBox)
 		{
 			CachedValue = Value.Get();
-			if (!MinValue.IsSet())
-			{
-				CachedValue -= Offset;
-			}
+			CachedValue -= Offset;
 
 			if (MinValue.IsSet() && CachedValue - Offset < MinValue.GetValue())
 			{
@@ -316,10 +313,7 @@ private:
 		if (EntryBox)
 		{
 			CachedValue = Value.Get();
-			if (!MaxValue.IsSet())
-			{
-				CachedValue += Offset;
-			}
+			CachedValue += Offset;
 
 			if (MaxValue.IsSet() && CachedValue + Offset > MaxValue.GetValue())
 			{
