Fix timed data monitor zoom being too slow by adding control modifer key to increase zoom scale.

diff --git a/Engine/Plugins/VirtualProduction/TimedDataMonitor/Source/TimedDataMonitorEditor/Private/STimedDataMonitorBufferVisualizer.cpp b/Engine/Plugins/VirtualProduction/TimedDataMonitor/Source/TimedDataMonitorEditor/Private/STimedDataMonitorBufferVisualizer.cpp
--- a/Engine/Plugins/VirtualProduction/TimedDataMonitor/Source/TimedDataMonitorEditor/Private/STimedDataMonitorBufferVisualizer.cpp
+++ b/Engine/Plugins/VirtualProduction/TimedDataMonitor/Source/TimedDataMonitorEditor/Private/STimedDataMonitorBufferVisualizer.cpp
@@ -3,25 +3,18 @@
 #include "STimedDataMonitorBufferVisualizer.h"
 
 #include "Engine/Engine.h"
+#include "Framework/Application/SlateApplication.h"
+#include "Framework/Notifications/NotificationManager.h"
 #include "Misc/App.h"
-#include "Modules/ModuleManager.h"
+#include "STimingDiagramWidget.h"
 #include "TimedDataMonitorEditorSettings.h"
 #include "TimedDataMonitorSubsystem.h"
-
-#include "Framework/MultiBox/MultiBoxBuilder.h"
-#include "Framework/Notifications/NotificationManager.h"
-
-#include "STimingDiagramWidget.h"
-#include "Widgets/Layout/SBorder.h"
-#include "Widgets/SBoxPanel.h"
-#include "Widgets/SOverlay.h"
 #include "Widgets/Layout/SBox.h"
 #include "Widgets/Layout/SScrollBorder.h"
 #include "Widgets/Views/STableRow.h"
 
 #define LOCTEXT_NAMESPACE "TimedDataBufferVisualizer"
 
-
 class STimedDataMonitorBufferVisualizerListView;
 
 /* TimedDataBufferVisualizer
@@ -127,6 +120,11 @@ public:
 	{
 		float Scale = 2.f;
 		float MinValue = 2.f;
+		if (FSlateApplication::Get().GetModifierKeys().IsControlDown())
+		{
+			Scale *= ControlKeyModifier;
+		}
+
 		SizePerSeconds = FMath::Max(MinValue, SizePerSeconds + MouseEvent.GetWheelDelta() * Scale);
 		return FReply::Handled();
 	}
@@ -137,6 +135,7 @@ public:
 	}
 
 	float SizePerSeconds = 100.f;
+	float ControlKeyModifier = 200.f;
 };
 
 
