Reduce default zoom modifier value for timed data monitor buffer visualizer
Add tooltips for minus and plus button of numeric entry box

diff --git a/Engine/Plugins/VirtualProduction/TimedDataMonitor/Source/TimedDataMonitorEditor/Private/STimedDataListView.cpp b/Engine/Plugins/VirtualProduction/TimedDataMonitor/Source/TimedDataMonitorEditor/Private/STimedDataListView.cpp
--- a/Engine/Plugins/VirtualProduction/TimedDataMonitor/Source/TimedDataMonitorEditor/Private/STimedDataListView.cpp
+++ b/Engine/Plugins/VirtualProduction/TimedDataMonitor/Source/TimedDataMonitorEditor/Private/STimedDataListView.cpp
@@ -354,6 +354,21 @@ TSharedRef<SWidget> STimedDataInputTableRow::GenerateWidgetForColumn(const FName
 				.CurrentValue_Lambda([this]() { return static_cast<int32>(Item->CachedInputEvaluationOffsetType); })
 				.OnEnumSelectionChanged(this, &STimedDataInputTableRow::OnEvaluationOffsetTypeChanged);
 
+			static const FString SecondString = TEXT("seconds");
+			static const FString FrameString = TEXT("frames");
+
+			auto MinusButtonTooltipTextLambda = [this]()
+			{
+				const TCHAR* UnitString = Item->CachedInputEvaluationOffsetType == ETimedDataInputEvaluationOffsetType::Seconds ? *SecondString : *FrameString;
+				return FText::Format(LOCTEXT("RushToolTip", "Rush this source (in {0})."), FText::FromStringView(UnitString));
+			};
+
+			auto PlusButtonTooltipTextLambda = [this]()
+			{
+				const TCHAR* UnitString = Item->CachedInputEvaluationOffsetType == ETimedDataInputEvaluationOffsetType::Seconds ? *SecondString : *FrameString;
+				return FText::Format(LOCTEXT("DelayToolTip", "Delay this source (in {0})."), FText::FromStringView(UnitString));
+			};
+
 			return SNew(SWidgetSwitcher)
 				.WidgetIndex_Lambda([this]()
 				{
@@ -368,20 +383,22 @@ TSharedRef<SWidget> STimedDataInputTableRow::GenerateWidgetForColumn(const FName
 				[
 					SNew(STimedDataNumericEntryBox<double>)
 						.TextStyle(ItemTextBlockStyle)
-						.ToolTipText(LOCTEXT("TimeCorrection_ToolTip", "Time Correction."))
 						.Value(this, &STimedDataInputTableRow::GetFloatEvaluationOffset)
 						.ComboButton(false)
 						.SuffixWidget(Label)
+						.MinusButtonToolTipText_Lambda(MinusButtonTooltipTextLambda)
+						.PlusButtonToolTipText_Lambda(PlusButtonTooltipTextLambda)
 						.OnValueCommitted(this, &STimedDataInputTableRow::SetEvaluationOffset)
 				]
 				+SWidgetSwitcher::Slot()
 				[
 					SNew(STimedDataNumericEntryBox<int32>)
 						.TextStyle(ItemTextBlockStyle)
-						.ToolTipText(LOCTEXT("TimeCorrection_ToolTip", "Time Correction."))
 						.Value(this, &STimedDataInputTableRow::GetIntEvaluationOffset)
 						.ComboButton(false)
 						.SuffixWidget(Label)
+						.MinusButtonToolTipText_Lambda(MinusButtonTooltipTextLambda)
+						.PlusButtonToolTipText_Lambda(PlusButtonTooltipTextLambda)
 						.OnValueCommitted(this, &STimedDataInputTableRow::SetEvaluationOffset)
 				];
 		}
diff --git a/Engine/Plugins/VirtualProduction/TimedDataMonitor/Source/TimedDataMonitorEditor/Private/STimedDataMonitorBufferVisualizer.cpp b/Engine/Plugins/VirtualProduction/TimedDataMonitor/Source/TimedDataMonitorEditor/Private/STimedDataMonitorBufferVisualizer.cpp
--- a/Engine/Plugins/VirtualProduction/TimedDataMonitor/Source/TimedDataMonitorEditor/Private/STimedDataMonitorBufferVisualizer.cpp
+++ b/Engine/Plugins/VirtualProduction/TimedDataMonitor/Source/TimedDataMonitorEditor/Private/STimedDataMonitorBufferVisualizer.cpp
@@ -120,6 +120,7 @@ public:
 	{
 		float Scale = 2.f;
 		float MinValue = 2.f;
+
 		if (FSlateApplication::Get().GetModifierKeys().IsControlDown())
 		{
 			Scale *= ControlKeyModifier;
@@ -135,7 +136,7 @@ public:
 	}
 
 	float SizePerSeconds = 100.f;
-	float ControlKeyModifier = 200.f;
+	float ControlKeyModifier = 50.f;
 };
 
 
diff --git a/Engine/Plugins/VirtualProduction/TimedDataMonitor/Source/TimedDataMonitorEditor/Private/STimedDataNumericEntryBox.h b/Engine/Plugins/VirtualProduction/TimedDataMonitor/Source/TimedDataMonitorEditor/Private/STimedDataNumericEntryBox.h
--- a/Engine/Plugins/VirtualProduction/TimedDataMonitor/Source/TimedDataMonitorEditor/Private/STimedDataNumericEntryBox.h
+++ b/Engine/Plugins/VirtualProduction/TimedDataMonitor/Source/TimedDataMonitorEditor/Private/STimedDataNumericEntryBox.h
@@ -18,7 +18,7 @@
 #include "Widgets/SNullWidget.h"
 #include "Widgets/Text/STextBlock.h"
 
-
+#define LOCTEXT_NAMESPACE "TimeDataNumericEntryBox"
 
 template<class NumericType>
 class STimedDataNumericEntryBox : public SCompoundWidget
@@ -64,6 +64,8 @@ public:
 		SLATE_ARGUMENT(NumericType, ControlModifierOffset)
 		SLATE_ARGUMENT(NumericType, AltModifierOffset)
 		SLATE_ARGUMENT(NumericType, ShiftModifierOffset)
+		SLATE_ATTRIBUTE(FText, MinusButtonToolTipText)
+		SLATE_ATTRIBUTE(FText, PlusButtonToolTipText)
 	SLATE_END_ARGS()
 	
 public:
@@ -81,6 +83,8 @@ public:
 		AltOffset = InArgs._AltModifierOffset;
 		ControlOffset = InArgs._ControlModifierOffset;
 		ShiftOffset = InArgs._ShiftModifierOffset;
+		MinusButtonToolTip = InArgs._MinusButtonToolTipText;
+		PlusButtonToolTip = InArgs._PlusButtonToolTipText;
 
 		TSharedRef<SWidget> TextBlock = SNew(STextBlock)
 			.TextStyle(InArgs._TextStyle)
@@ -190,6 +194,7 @@ private:
 			[
 				SNew(SButton)
 				.ButtonStyle(FTimedDataMonitorEditorStyle::Get(), "FlatButton")
+				.ToolTipText(MinusButtonToolTip)
 				.OnClicked(this, &STimedDataNumericEntryBox::OnMinusClicked)
 				.Content()
 				[
@@ -223,6 +228,7 @@ private:
 			[
 				SNew(SButton)
 				.ButtonStyle(FTimedDataMonitorEditorStyle::Get(), "FlatButton")
+				.ToolTipText(PlusButtonToolTip)
 				.OnClicked(this, &STimedDataNumericEntryBox::OnPlusClicked)
 				.Content()
 				[
@@ -366,5 +372,9 @@ private:
 	NumericType AltOffset;
 	NumericType ShiftOffset;
 	NumericType ControlOffset;
+	TAttribute<FText> MinusButtonToolTip;
+	TAttribute<FText> PlusButtonToolTip;
 };
 
+
+#undef LOCTEXT_NAMESPACE
\ No newline at end of file
