Timed Data Monitor : Fix idle time showing negative frame numbers

#rnx

diff --git a/Engine/Plugins/VirtualProduction/TimedDataMonitor/Source/TimedDataMonitorEditor/Private/STimedDataGenlock.cpp b/Engine/Plugins/VirtualProduction/TimedDataMonitor/Source/TimedDataMonitorEditor/Private/STimedDataGenlock.cpp
--- a/Engine/Plugins/VirtualProduction/TimedDataMonitor/Source/TimedDataMonitorEditor/Private/STimedDataGenlock.cpp
+++ b/Engine/Plugins/VirtualProduction/TimedDataMonitor/Source/TimedDataMonitorEditor/Private/STimedDataGenlock.cpp
@@ -182,7 +182,7 @@ void STimedDataGenlock::UpdateCachedValue(float InDeltaTime)
 	uint32 IdleFrames = 0;
 	if (!FMath::IsNearlyZero(IdleTime))
 	{
-		IdleFrames = FMath::RoundToInt(IdleTime / (TargetFrameTime * TargetFrameTime) );
+		IdleFrames = FMath::Abs(FMath::RoundToInt(IdleTime / (TargetFrameTime * TargetFrameTime) ));
 	}
 
 	CachedFPSText = FText::AsNumber(1.0/FApp::GetDeltaTime(), &FPSFormattingOptions);
