Fix compile error in STimedDataListView

#rnx

diff --git a/Engine/Plugins/VirtualProduction/TimedDataMonitor/Source/TimedDataMonitorEditor/Private/STimedDataListView.cpp b/Engine/Plugins/VirtualProduction/TimedDataMonitor/Source/TimedDataMonitorEditor/Private/STimedDataListView.cpp
--- a/Engine/Plugins/VirtualProduction/TimedDataMonitor/Source/TimedDataMonitorEditor/Private/STimedDataListView.cpp
+++ b/Engine/Plugins/VirtualProduction/TimedDataMonitor/Source/TimedDataMonitorEditor/Private/STimedDataListView.cpp
@@ -266,7 +266,7 @@ TSharedRef<SWidget> STimedDataInputTableRow::GenerateWidgetForColumn(const FName
 				.AutoWidth()
 				.Padding(6, 0, 0, 0)
 				[
-					SNew(SExpanderArrow, SharedThis(this))
+					SNew(SExpanderArrow, SharedThis<STimedDataInputTableRow>(this))
 					.ShouldDrawWires(false)
 					.IndentAmount(12)
 				]
