LiveLinkHub - Enable Timed Data Monitor

#rnx

diff --git a/Engine/Plugins/Animation/LiveLinkHub/LiveLinkHub.uplugin b/Engine/Plugins/Animation/LiveLinkHub/LiveLinkHub.uplugin
--- a/Engine/Plugins/Animation/LiveLinkHub/LiveLinkHub.uplugin
+++ b/Engine/Plugins/Animation/LiveLinkHub/LiveLinkHub.uplugin
@@ -29,9 +29,13 @@
 			"Name": "UdpMessaging",
 			"Enabled": true
 		},
-        {
+		{
 			"Name": "MessageBusTester",
 			"Enabled": true
+		},
+		{
+			"Name": "TimedDataMonitor",
+			"Enabled": true
 		}
 	],
 	"Modules": [
diff --git a/Engine/Plugins/Animation/LiveLinkHub/Source/LiveLinkHub/LiveLinkHub.Build.cs b/Engine/Plugins/Animation/LiveLinkHub/Source/LiveLinkHub/LiveLinkHub.Build.cs
--- a/Engine/Plugins/Animation/LiveLinkHub/Source/LiveLinkHub/LiveLinkHub.Build.cs
+++ b/Engine/Plugins/Animation/LiveLinkHub/Source/LiveLinkHub/LiveLinkHub.Build.cs
@@ -48,7 +48,9 @@ public class LiveLinkHub : ModuleRules
 				"UnrealEd",
 				"OutputLog",
 				"NamingTokens",
-				"MessageBusTesterEditor"
+				"MessageBusTesterEditor",
+				"TimedDataMonitorEditor"
+
 			});
 
 		if (Target.Configuration != UnrealTargetConfiguration.Shipping)
diff --git a/Engine/Plugins/Animation/LiveLinkHub/Source/LiveLinkHub/Private/LiveLinkHubApplication.cpp b/Engine/Plugins/Animation/LiveLinkHub/Source/LiveLinkHub/Private/LiveLinkHubApplication.cpp
--- a/Engine/Plugins/Animation/LiveLinkHub/Source/LiveLinkHub/Private/LiveLinkHubApplication.cpp
+++ b/Engine/Plugins/Animation/LiveLinkHub/Source/LiveLinkHub/Private/LiveLinkHubApplication.cpp
@@ -25,6 +25,7 @@
 #	endif
 #	include "ISlateReflectorModule.h"
 #	include "IMessageBusTesterEditorModule.h"
+#	include "TimedDataMonitorEditorModule.h"
 #endif
 
 #define LOCTEXT_NAMESPACE "LiveLinkHubApplication"
@@ -560,7 +561,7 @@ void FLiveLinkHubApplication::CreateToolsMenu(FMenuBuilder& MenuBuilder, UToolMe
 
 #if WITH_UNREAL_TARGET_DEVELOPER_TOOLS
 	MenuBuilder.AddMenuEntry(
-		LOCTEXT("FileMenu_Development_AutomationTools", "Automation Tools"),
+		LOCTEXT("ToolsMenu_Development_AutomationTools", "Automation Tools"),
 		FText(),
 		FSlateIcon(),
 		FExecuteAction::CreateStatic([]()
@@ -572,7 +573,7 @@ void FLiveLinkHubApplication::CreateToolsMenu(FMenuBuilder& MenuBuilder, UToolMe
 	);
 #endif // #if WITH_UNREAL_TARGET_DEVELOPER_TOOLS
 	MenuBuilder.AddMenuEntry(
-		LOCTEXT("FileMenu_Development_WidgetReflector", "Widget Reflector"),
+		LOCTEXT("ToolsMenu_Development_WidgetReflector", "Widget Reflector"),
 		FText(),
 		FSlateIcon(),
 		FExecuteAction::CreateStatic([]()
@@ -583,7 +584,7 @@ void FLiveLinkHubApplication::CreateToolsMenu(FMenuBuilder& MenuBuilder, UToolMe
 			})
 	);
 	MenuBuilder.AddMenuEntry(
-		LOCTEXT("FileMenu_Development_MesageBusTester", "Message Bus Tester"),
+		LOCTEXT("ToolsMenu_Development_MesageBusTester", "Message Bus Tester"),
 		FText(),
 		FSlateIcon(),
 		FExecuteAction::CreateStatic([]()
@@ -592,6 +593,18 @@ void FLiveLinkHubApplication::CreateToolsMenu(FMenuBuilder& MenuBuilder, UToolMe
 				Module.DisplayMessageBusTester();
 			})
 	);
+
+	MenuBuilder.AddMenuEntry(
+		LOCTEXT("ToolMenu_Development_TimedDataMonitor", "Timed Data Monitor"),
+		FText(),
+		FSlateIcon(),
+		FExecuteAction::CreateLambda([this]()
+		{
+			FTimedDataMonitorEditorModule& TimedDataMonitorModule = FModuleManager::LoadModuleChecked<FTimedDataMonitorEditorModule>("TimedDataMonitorEditor");
+			TimedDataMonitorModule.DisplayTimedDataMonitorPanel(TabManager);
+		})
+	);
+
 	MenuBuilder.EndSection();
 
 
@@ -612,6 +625,9 @@ TSharedPtr<SWidget> FLiveLinkHubApplication::CreateMenuBar(const TSharedPtr<FTab
 {
 	RegisterMainMenu();
 
+	FTimedDataMonitorEditorModule& TimedDataMonitorModule = FModuleManager::LoadModuleChecked<FTimedDataMonitorEditorModule>("TimedDataMonitorEditor");
+	TimedDataMonitorModule.RegisterNomadTabSpawner(TabManager);
+
 	InToolMenuContext.AppendCommandList(GetToolkitCommands());
 
 	USlateTabManagerContext* ContextObject = NewObject<USlateTabManagerContext>();
diff --git a/Engine/Plugins/VirtualProduction/TimedDataMonitor/Source/TimedDataMonitorEditor/Private/STimedDataMonitorPanel.cpp b/Engine/Plugins/VirtualProduction/TimedDataMonitor/Source/TimedDataMonitorEditor/Private/STimedDataMonitorPanel.cpp
--- a/Engine/Plugins/VirtualProduction/TimedDataMonitor/Source/TimedDataMonitorEditor/Private/STimedDataMonitorPanel.cpp
+++ b/Engine/Plugins/VirtualProduction/TimedDataMonitor/Source/TimedDataMonitorEditor/Private/STimedDataMonitorPanel.cpp
@@ -33,14 +33,11 @@
 
 
 TWeakPtr<STimedDataMonitorPanel> STimedDataMonitorPanel::WidgetInstance;
-FDelegateHandle STimedDataMonitorPanel::LevelEditorTabManagerChangedHandle;
-
 
 namespace Utilities
 {
 	static const FName NAME_App = FName("TimedDataSourceMonitorPanelApp");
 	static const FName NAME_TimedDataMonitorBuffers = FName("TimedDataMonitorBuffers");
-	static const FName NAME_LevelEditorModuleName("LevelEditor");
 	static const FName NAME_LogName = "Timed Data Monitor";
 
 	TSharedRef<SDockTab> CreateApp(const FSpawnTabArgs& Args)
@@ -63,55 +60,33 @@ namespace Utilities
 }
 
 
-void STimedDataMonitorPanel::RegisterNomadTabSpawner(TSharedRef<FWorkspaceItem> InWorkspaceItem)
+const FName STimedDataMonitorPanel::TabName = "TimedDataSourceMonitorPanelApp";
+
+
+void STimedDataMonitorPanel::RegisterNomadTabSpawner(TSharedPtr<FTabManager> TabManager, TSharedRef<FWorkspaceItem> InWorkspaceItem)
 {
-	auto RegisterTabSpawner = [InWorkspaceItem]()
+	if (TabManager)
 	{
-		FLevelEditorModule& LevelEditorModule = FModuleManager::LoadModuleChecked<FLevelEditorModule>(Utilities::NAME_LevelEditorModuleName);
-		TSharedPtr<FTabManager> LevelEditorTabManager = LevelEditorModule.GetLevelEditorTabManager();
-
-		LevelEditorTabManager->RegisterTabSpawner(Utilities::NAME_App, FOnSpawnTab::CreateStatic(&Utilities::CreateApp))
+		TabManager->RegisterTabSpawner(Utilities::NAME_App, FOnSpawnTab::CreateStatic(&Utilities::CreateApp))
 			.SetDisplayName(LOCTEXT("TabTitle", "Timed Data Monitor"))
 			.SetTooltipText(LOCTEXT("TooltipText", "Monitor inputs that can be time synchronized"))
 			.SetGroup(InWorkspaceItem)
 			.SetIcon(FSlateIcon(FTimedDataMonitorEditorStyle::Get().GetStyleSetName(), "Img.TimedDataMonitor.Icon"));
-
-		FGlobalTabmanager::Get()->RegisterNomadTabSpawner(Utilities::NAME_TimedDataMonitorBuffers, FOnSpawnTab::CreateStatic(&Utilities::CreateBuffersVisualizer))
-			.SetDisplayName(LOCTEXT("BufferVisualizerTitle", "Buffer Visualizer"))
-			.SetTooltipText(LOCTEXT("BufferVisualizerTooltip", "Open buffer visualizer tab for Timed Data Monitor."))
-			.SetMenuType(ETabSpawnerMenuType::Hidden);
-	};
-
-	FLevelEditorModule& LevelEditorModule = FModuleManager::LoadModuleChecked<FLevelEditorModule>(Utilities::NAME_LevelEditorModuleName);
-	if (LevelEditorModule.GetLevelEditorTabManager())
-	{
-		RegisterTabSpawner();
 	}
-	else
-	{
-		LevelEditorTabManagerChangedHandle = LevelEditorModule.OnTabManagerChanged().AddLambda(RegisterTabSpawner);
-	}
-}
 
+	FGlobalTabmanager::Get()->RegisterNomadTabSpawner(Utilities::NAME_TimedDataMonitorBuffers, FOnSpawnTab::CreateStatic(&Utilities::CreateBuffersVisualizer))
+		.SetDisplayName(LOCTEXT("BufferVisualizerTitle", "Buffer Visualizer"))
+		.SetTooltipText(LOCTEXT("BufferVisualizerTooltip", "Open buffer visualizer tab for Timed Data Monitor."))
+		.SetMenuType(ETabSpawnerMenuType::Hidden);
+}
 
-void STimedDataMonitorPanel::UnregisterNomadTabSpawner()
+void STimedDataMonitorPanel::UnregisterNomadTabSpawner(TSharedPtr<FTabManager> TabManager)
 {
 	FGlobalTabmanager::Get()->UnregisterNomadTabSpawner(Utilities::NAME_TimedDataMonitorBuffers);
 
-	if (FSlateApplication::IsInitialized() && FModuleManager::Get().IsModuleLoaded(Utilities::NAME_LevelEditorModuleName))
+	if (TabManager)
 	{
-		FLevelEditorModule* LevelEditorModule = FModuleManager::GetModulePtr<FLevelEditorModule>(Utilities::NAME_LevelEditorModuleName);
-		TSharedPtr<FTabManager> LevelEditorTabManager;
-		if (LevelEditorModule)
-		{
-			LevelEditorTabManager = LevelEditorModule->GetLevelEditorTabManager();
-			LevelEditorModule->OnTabManagerChanged().Remove(LevelEditorTabManagerChangedHandle);
-		}
-
-		if (LevelEditorTabManager.IsValid())
-		{
-			LevelEditorTabManager->UnregisterTabSpawner(Utilities::NAME_App);
-		}
+		TabManager->UnregisterTabSpawner(Utilities::NAME_App);
 	}
 }
 
diff --git a/Engine/Plugins/VirtualProduction/TimedDataMonitor/Source/TimedDataMonitorEditor/Private/STimedDataMonitorPanel.h b/Engine/Plugins/VirtualProduction/TimedDataMonitor/Source/TimedDataMonitorEditor/Private/STimedDataMonitorPanel.h
--- a/Engine/Plugins/VirtualProduction/TimedDataMonitor/Source/TimedDataMonitorEditor/Private/STimedDataMonitorPanel.h
+++ b/Engine/Plugins/VirtualProduction/TimedDataMonitor/Source/TimedDataMonitorEditor/Private/STimedDataMonitorPanel.h
@@ -14,6 +14,7 @@ enum class ETimedDataMonitorTimeCorrectionReturnCode : uint8;
 
 struct FSlateBrush;
 struct FTimedDataMonitorCalibrationResult;
+class FTabManager;
 class FTimedDataMonitorCalibration;
 class FWorkspaceItem;
 class IMessageLogListing;
@@ -27,9 +28,10 @@ enum class ETimedDataMonitorEvaluationState : uint8;
 class STimedDataMonitorPanel : public SCompoundWidget
 {
 public:
-	static void RegisterNomadTabSpawner(TSharedRef<FWorkspaceItem> InWorkspaceItem);
-	static void UnregisterNomadTabSpawner();
+	static void RegisterNomadTabSpawner(TSharedPtr<FTabManager> TabManager, TSharedRef<FWorkspaceItem> InWorkspaceItem);
+	static void UnregisterNomadTabSpawner(TSharedPtr<FTabManager> TabManager);
 	static TSharedPtr<STimedDataMonitorPanel> GetPanelInstance();
+	static const FName TabName;
 
 private:
 	using Super = SCompoundWidget;
@@ -99,6 +101,4 @@ private:
 
 	bool bRefreshRequested = true;
 	double LastCachedValueUpdateTime = 0.0;
-
-	static FDelegateHandle LevelEditorTabManagerChangedHandle;
 };
diff --git a/Engine/Plugins/VirtualProduction/TimedDataMonitor/Source/TimedDataMonitorEditor/Private/TimedDataMonitorEditorModule.cpp b/Engine/Plugins/VirtualProduction/TimedDataMonitor/Source/TimedDataMonitorEditor/Private/TimedDataMonitorEditorModule.cpp
--- a/Engine/Plugins/VirtualProduction/TimedDataMonitor/Source/TimedDataMonitorEditor/Private/TimedDataMonitorEditorModule.cpp
+++ b/Engine/Plugins/VirtualProduction/TimedDataMonitor/Source/TimedDataMonitorEditor/Private/TimedDataMonitorEditorModule.cpp
@@ -1,45 +1,96 @@
 // Copyright Epic Games, Inc. All Rights Reserved.
 
+#include "TimedDataMonitorEditorModule.h"
 
+#include "Framework/Application/SlateApplication.h"
 #include "ISettingsModule.h"
+#include "LevelEditor.h"
 #include "STimedDataMonitorPanel.h"
+#include "STimedDataMonitorBufferVisualizer.h"
 #include "WorkspaceMenuStructureModule.h"
 #include "WorkspaceMenuStructure.h"
 
-
 #define LOCTEXT_NAMESPACE "TimedDataMonitorEditorModule"
 
+namespace UE::TDM
+{
+	static const FName NAME_LevelEditorModuleName("LevelEditor");
+}
+
 
-class FTimedDataMonitorEditorModule : public IModuleInterface
+void FTimedDataMonitorEditorModule::StartupModule()
 {
-public:
-	virtual void StartupModule() override
+	if (ISettingsModule* SettingsModule = FModuleManager::GetModulePtr<ISettingsModule>("Settings"))
 	{
-		STimedDataMonitorPanel::RegisterNomadTabSpawner(WorkspaceMenu::GetMenuStructure().GetLevelEditorVirtualProductionCategory());
+		SettingsModule->RegisterSettings("Editor", "Plugins", "Timed Data Monitor",
+			LOCTEXT("SettingsName", "Timed Data Monitor"),
+			LOCTEXT("Description", "Configure the Timed Data Monitor panel."),
+			GetMutableDefault<UTimedDataMonitorEditorSettings>()
+		);
+	}
+
+	FLevelEditorModule& LevelEditorModule = FModuleManager::LoadModuleChecked<FLevelEditorModule>(UE::TDM::NAME_LevelEditorModuleName);
+	TSharedPtr<FTabManager> LevelEditorTabManager = LevelEditorModule.GetLevelEditorTabManager();
 
-		if (ISettingsModule* SettingsModule = FModuleManager::GetModulePtr<ISettingsModule>("Settings"))
+	auto RegisterLambda = [this]()
 		{
-			SettingsModule->RegisterSettings("Editor", "Plugins", "Timed Data Monitor",
-				LOCTEXT("SettingsName", "Timed Data Monitor"),
-				LOCTEXT("Description", "Configure the Timed Data Monitor panel."),
-				GetMutableDefault<UTimedDataMonitorEditorSettings>()
-			);
-		}
+			FLevelEditorModule& LevelEditorModule = FModuleManager::LoadModuleChecked<FLevelEditorModule>(UE::TDM::NAME_LevelEditorModuleName);
+			TSharedPtr<FTabManager> LevelEditorTabManager = LevelEditorModule.GetLevelEditorTabManager();
+			if (LevelEditorTabManager)
+			{
+				RegisterNomadTabSpawner(LevelEditorTabManager);
+			}
+		};
+
+	if (LevelEditorTabManager)
+	{
+		RegisterLambda();
+	}
+	else
+	{
+		LevelEditorTabManagerChangedHandle = LevelEditorModule.OnTabManagerChanged().AddLambda(RegisterLambda);
 	}
+}
 
-	virtual void ShutdownModule() override
+void FTimedDataMonitorEditorModule::ShutdownModule()
+{
+	if (!IsRunningCommandlet() && UObjectInitialized() && !IsEngineExitRequested())
 	{
-		if (!IsRunningCommandlet() && UObjectInitialized() && !IsEngineExitRequested())
+		ISettingsModule* SettingsModule = FModuleManager::GetModulePtr<ISettingsModule>("Settings");
+		if (SettingsModule != nullptr)
+		{
+			SettingsModule->UnregisterSettings("Editor", "Plugins", "Timed Data Monitor");
+		}
+
+		if (FSlateApplication::IsInitialized())
 		{
-			ISettingsModule* SettingsModule = FModuleManager::GetModulePtr<ISettingsModule>("Settings");
-			if (SettingsModule != nullptr)
+			if (FLevelEditorModule* LevelEditorModule = FModuleManager::GetModulePtr<FLevelEditorModule>(UE::TDM::NAME_LevelEditorModuleName))
 			{
-				SettingsModule->UnregisterSettings("Editor", "Plugins", "Timed Data Monitor");
+				LevelEditorModule->OnTabManagerChanged().Remove(LevelEditorTabManagerChangedHandle);
+
+				if (TSharedPtr<FTabManager> LevelEditorTabManager = LevelEditorModule->GetLevelEditorTabManager())
+				{
+					UnregisterNomadTabSpawner(LevelEditorTabManager);
+				}
 			}
-			STimedDataMonitorPanel::UnregisterNomadTabSpawner();
 		}
 	}
-};
+}
+
+void FTimedDataMonitorEditorModule::RegisterNomadTabSpawner(TSharedPtr<FTabManager> TabManager)
+{
+	STimedDataMonitorPanel::RegisterNomadTabSpawner(TabManager, WorkspaceMenu::GetMenuStructure().GetLevelEditorVirtualProductionCategory());
+}
+
+void FTimedDataMonitorEditorModule::UnregisterNomadTabSpawner(TSharedPtr<FTabManager> TabManager)
+{
+	STimedDataMonitorPanel::UnregisterNomadTabSpawner(TabManager);
+}
+
+void FTimedDataMonitorEditorModule::DisplayTimedDataMonitorPanel(TSharedPtr<FTabManager> TabManager)
+{
+	TabManager->TryInvokeTab(STimedDataMonitorPanel::TabName);
+}
 
 IMPLEMENT_MODULE(FTimedDataMonitorEditorModule, TimedDataMonitorEditor);
 
diff --git a/Engine/Plugins/VirtualProduction/TimedDataMonitor/Source/TimedDataMonitorEditor/Public/TimedDataMonitorEditorModule.h b/Engine/Plugins/VirtualProduction/TimedDataMonitor/Source/TimedDataMonitorEditor/Public/TimedDataMonitorEditorModule.h
new file mode 100644
--- /dev/null
+++ b/Engine/Plugins/VirtualProduction/TimedDataMonitor/Source/TimedDataMonitorEditor/Public/TimedDataMonitorEditorModule.h
@@ -0,0 +1,28 @@
+// Copyright Epic Games, Inc. All Rights Reserved.
+#pragma once
+
+#include "Modules/ModuleInterface.h"
+
+#include "Delegates/IDelegateInstance.h"
+#include "Templates/SharedPointer.h"
+
+class FTabManager;
+class SWidget;
+
+class FTimedDataMonitorEditorModule : public IModuleInterface
+{
+public:
+	virtual void StartupModule() override;
+	virtual void ShutdownModule() override;
+
+	/** Register a tab spawner for the TimedDataMonitorPanel. */
+	TIMEDDATAMONITOREDITOR_API void RegisterNomadTabSpawner(TSharedPtr<FTabManager> TabManager);
+	/** Unregister a tab spawner for the TimedDataMonitorPanel. */
+	TIMEDDATAMONITOREDITOR_API void UnregisterNomadTabSpawner(TSharedPtr<FTabManager> TabManager);
+	/** Try invoking the TimedDataMonitor tab. */
+	TIMEDDATAMONITOREDITOR_API void DisplayTimedDataMonitorPanel(TSharedPtr<FTabManager> TabManager);
+private:
+	/** Handle to the delegate used to register the tab spawner when the level editor tab manager changes. */
+	FDelegateHandle LevelEditorTabManagerChangedHandle;
+};
+
diff --git a/Engine/Source/Programs/LiveLinkHub/Source/LiveLinkHub.Target.cs b/Engine/Source/Programs/LiveLinkHub/Source/LiveLinkHub.Target.cs
--- a/Engine/Source/Programs/LiveLinkHub/Source/LiveLinkHub.Target.cs
+++ b/Engine/Source/Programs/LiveLinkHub/Source/LiveLinkHub.Target.cs
@@ -65,7 +65,9 @@ public class LiveLinkHubTarget : TargetRules
 			"MediaIOFramework",
 			"MediaFrameworkUtilities",
 			"AjaMedia",
-			"BlackmagicMedia"
+			"BlackmagicMedia",
+
+			"TimedDataMonitor"
 		});
 
 		if (bEnableMetaHumanLiveLinkPlugin)
