[Backout] - CL47910311
#rnx
Original CL Desc
-----------------------------------------------------------------
LiveLinkHub - Enable Timed Data Monitor

#rnx

diff --git a/Engine/Plugins/Animation/LiveLinkHub/LiveLinkHub.uplugin b/Engine/Plugins/Animation/LiveLinkHub/LiveLinkHub.uplugin
--- a/Engine/Plugins/Animation/LiveLinkHub/LiveLinkHub.uplugin
+++ b/Engine/Plugins/Animation/LiveLinkHub/LiveLinkHub.uplugin
@@ -29,13 +29,9 @@
 			"Name": "UdpMessaging",
 			"Enabled": true
 		},
-		{
+        {
 			"Name": "MessageBusTester",
 			"Enabled": true
-		},
-		{
-			"Name": "TimedDataMonitor",
-			"Enabled": true
 		}
 	],
 	"Modules": [
diff --git a/Engine/Plugins/Animation/LiveLinkHub/Source/LiveLinkHub/LiveLinkHub.Build.cs b/Engine/Plugins/Animation/LiveLinkHub/Source/LiveLinkHub/LiveLinkHub.Build.cs
--- a/Engine/Plugins/Animation/LiveLinkHub/Source/LiveLinkHub/LiveLinkHub.Build.cs
+++ b/Engine/Plugins/Animation/LiveLinkHub/Source/LiveLinkHub/LiveLinkHub.Build.cs
@@ -48,9 +48,7 @@ public class LiveLinkHub : ModuleRules
 				"UnrealEd",
 				"OutputLog",
 				"NamingTokens",
-				"MessageBusTesterEditor",
-				"TimedDataMonitorEditor"
-
+				"MessageBusTesterEditor"
 			});
 
 		if (Target.Configuration != UnrealTargetConfiguration.Shipping)
diff --git a/Engine/Plugins/Animation/LiveLinkHub/Source/LiveLinkHub/Private/LiveLinkHubApplication.cpp b/Engine/Plugins/Animation/LiveLinkHub/Source/LiveLinkHub/Private/LiveLinkHubApplication.cpp
--- a/Engine/Plugins/Animation/LiveLinkHub/Source/LiveLinkHub/Private/LiveLinkHubApplication.cpp
+++ b/Engine/Plugins/Animation/LiveLinkHub/Source/LiveLinkHub/Private/LiveLinkHubApplication.cpp
@@ -25,7 +25,6 @@
 #	endif
 #	include "ISlateReflectorModule.h"
 #	include "IMessageBusTesterEditorModule.h"
-#	include "TimedDataMonitorEditorModule.h"
 #endif
 
 #define LOCTEXT_NAMESPACE "LiveLinkHubApplication"
@@ -561,7 +560,7 @@ void FLiveLinkHubApplication::CreateToolsMenu(FMenuBuilder& MenuBuilder, UToolMe
 
 #if WITH_UNREAL_TARGET_DEVELOPER_TOOLS
 	MenuBuilder.AddMenuEntry(
-		LOCTEXT("ToolsMenu_Development_AutomationTools", "Automation Tools"),
+		LOCTEXT("FileMenu_Development_AutomationTools", "Automation Tools"),
 		FText(),
 		FSlateIcon(),
 		FExecuteAction::CreateStatic([]()
@@ -573,7 +572,7 @@ void FLiveLinkHubApplication::CreateToolsMenu(FMenuBuilder& MenuBuilder, UToolMe
 	);
 #endif // #if WITH_UNREAL_TARGET_DEVELOPER_TOOLS
 	MenuBuilder.AddMenuEntry(
-		LOCTEXT("ToolsMenu_Development_WidgetReflector", "Widget Reflector"),
+		LOCTEXT("FileMenu_Development_WidgetReflector", "Widget Reflector"),
 		FText(),
 		FSlateIcon(),
 		FExecuteAction::CreateStatic([]()
@@ -584,7 +583,7 @@ void FLiveLinkHubApplication::CreateToolsMenu(FMenuBuilder& MenuBuilder, UToolMe
 			})
 	);
 	MenuBuilder.AddMenuEntry(
-		LOCTEXT("ToolsMenu_Development_MesageBusTester", "Message Bus Tester"),
+		LOCTEXT("FileMenu_Development_MesageBusTester", "Message Bus Tester"),
 		FText(),
 		FSlateIcon(),
 		FExecuteAction::CreateStatic([]()
@@ -593,18 +592,6 @@ void FLiveLinkHubApplication::CreateToolsMenu(FMenuBuilder& MenuBuilder, UToolMe
 				Module.DisplayMessageBusTester();
 			})
 	);
-
-	MenuBuilder.AddMenuEntry(
-		LOCTEXT("ToolMenu_Development_TimedDataMonitor", "Timed Data Monitor"),
-		FText(),
-		FSlateIcon(),
-		FExecuteAction::CreateLambda([this]()
-		{
-			FTimedDataMonitorEditorModule& TimedDataMonitorModule = FModuleManager::LoadModuleChecked<FTimedDataMonitorEditorModule>("TimedDataMonitorEditor");
-			TimedDataMonitorModule.DisplayTimedDataMonitorPanel(TabManager);
-		})
-	);
-
 	MenuBuilder.EndSection();
 
 
@@ -625,9 +612,6 @@ TSharedPtr<SWidget> FLiveLinkHubApplication::CreateMenuBar(const TSharedPtr<FTab
 {
 	RegisterMainMenu();
 
-	FTimedDataMonitorEditorModule& TimedDataMonitorModule = FModuleManager::LoadModuleChecked<FTimedDataMonitorEditorModule>("TimedDataMonitorEditor");
-	TimedDataMonitorModule.RegisterNomadTabSpawner(TabManager);
-
 	InToolMenuContext.AppendCommandList(GetToolkitCommands());
 
 	USlateTabManagerContext* ContextObject = NewObject<USlateTabManagerContext>();
diff --git a/Engine/Plugins/VirtualProduction/TimedDataMonitor/Source/TimedDataMonitorEditor/Private/STimedDataMonitorPanel.cpp b/Engine/Plugins/VirtualProduction/TimedDataMonitor/Source/TimedDataMonitorEditor/Private/STimedDataMonitorPanel.cpp
--- a/Engine/Plugins/VirtualProduction/TimedDataMonitor/Source/TimedDataMonitorEditor/Private/STimedDataMonitorPanel.cpp
+++ b/Engine/Plugins/VirtualProduction/TimedDataMonitor/Source/TimedDataMonitorEditor/Private/STimedDataMonitorPanel.cpp
@@ -33,11 +33,14 @@
 
 
 TWeakPtr<STimedDataMonitorPanel> STimedDataMonitorPanel::WidgetInstance;
+FDelegateHandle STimedDataMonitorPanel::LevelEditorTabManagerChangedHandle;
+
 
 namespace Utilities
 {
 	static const FName NAME_App = FName("TimedDataSourceMonitorPanelApp");
 	static const FName NAME_TimedDataMonitorBuffers = FName("TimedDataMonitorBuffers");
+	static const FName NAME_LevelEditorModuleName("LevelEditor");
 	static const FName NAME_LogName = "Timed Data Monitor";
 
 	TSharedRef<SDockTab> CreateApp(const FSpawnTabArgs& Args)
@@ -60,33 +63,55 @@ namespace Utilities
 }
 
 
-const FName STimedDataMonitorPanel::TabName = "TimedDataSourceMonitorPanelApp";
-
-
-void STimedDataMonitorPanel::RegisterNomadTabSpawner(TSharedPtr<FTabManager> TabManager, TSharedRef<FWorkspaceItem> InWorkspaceItem)
+void STimedDataMonitorPanel::RegisterNomadTabSpawner(TSharedRef<FWorkspaceItem> InWorkspaceItem)
 {
-	if (TabManager)
+	auto RegisterTabSpawner = [InWorkspaceItem]()
 	{
-		TabManager->RegisterTabSpawner(Utilities::NAME_App, FOnSpawnTab::CreateStatic(&Utilities::CreateApp))
+		FLevelEditorModule& LevelEditorModule = FModuleManager::LoadModuleChecked<FLevelEditorModule>(Utilities::NAME_LevelEditorModuleName);
+		TSharedPtr<FTabManager> LevelEditorTabManager = LevelEditorModule.GetLevelEditorTabManager();
+
+		LevelEditorTabManager->RegisterTabSpawner(Utilities::NAME_App, FOnSpawnTab::CreateStatic(&Utilities::CreateApp))
 			.SetDisplayName(LOCTEXT("TabTitle", "Timed Data Monitor"))
 			.SetTooltipText(LOCTEXT("TooltipText", "Monitor inputs that can be time synchronized"))
 			.SetGroup(InWorkspaceItem)
 			.SetIcon(FSlateIcon(FTimedDataMonitorEditorStyle::Get().GetStyleSetName(), "Img.TimedDataMonitor.Icon"));
-	}
 
-	FGlobalTabmanager::Get()->RegisterNomadTabSpawner(Utilities::NAME_TimedDataMonitorBuffers, FOnSpawnTab::CreateStatic(&Utilities::CreateBuffersVisualizer))
-		.SetDisplayName(LOCTEXT("BufferVisualizerTitle", "Buffer Visualizer"))
-		.SetTooltipText(LOCTEXT("BufferVisualizerTooltip", "Open buffer visualizer tab for Timed Data Monitor."))
-		.SetMenuType(ETabSpawnerMenuType::Hidden);
+		FGlobalTabmanager::Get()->RegisterNomadTabSpawner(Utilities::NAME_TimedDataMonitorBuffers, FOnSpawnTab::CreateStatic(&Utilities::CreateBuffersVisualizer))
+			.SetDisplayName(LOCTEXT("BufferVisualizerTitle", "Buffer Visualizer"))
+			.SetTooltipText(LOCTEXT("BufferVisualizerTooltip", "Open buffer visualizer tab for Timed Data Monitor."))
+			.SetMenuType(ETabSpawnerMenuType::Hidden);
+	};
+
+	FLevelEditorModule& LevelEditorModule = FModuleManager::LoadModuleChecked<FLevelEditorModule>(Utilities::NAME_LevelEditorModuleName);
+	if (LevelEditorModule.GetLevelEditorTabManager())
+	{
+		RegisterTabSpawner();
+	}
+	else
+	{
+		LevelEditorTabManagerChangedHandle = LevelEditorModule.OnTabManagerChanged().AddLambda(RegisterTabSpawner);
+	}
 }
 
-void STimedDataMonitorPanel::UnregisterNomadTabSpawner(TSharedPtr<FTabManager> TabManager)
+
+void STimedDataMonitorPanel::UnregisterNomadTabSpawner()
 {
 	FGlobalTabmanager::Get()->UnregisterNomadTabSpawner(Utilities::NAME_TimedDataMonitorBuffers);
 
-	if (TabManager)
+	if (FSlateApplication::IsInitialized() && FModuleManager::Get().IsModuleLoaded(Utilities::NAME_LevelEditorModuleName))
 	{
-		TabManager->UnregisterTabSpawner(Utilities::NAME_App);
+		FLevelEditorModule* LevelEditorModule = FModuleManager::GetModulePtr<FLevelEditorModule>(Utilities::NAME_LevelEditorModuleName);
+		TSharedPtr<FTabManager> LevelEditorTabManager;
+		if (LevelEditorModule)
+		{
+			LevelEditorTabManager = LevelEditorModule->GetLevelEditorTabManager();
+			LevelEditorModule->OnTabManagerChanged().Remove(LevelEditorTabManagerChangedHandle);
+		}
+
+		if (LevelEditorTabManager.IsValid())
+		{
+			LevelEditorTabManager->UnregisterTabSpawner(Utilities::NAME_App);
+		}
 	}
 }
 
diff --git a/Engine/Plugins/VirtualProduction/TimedDataMonitor/Source/TimedDataMonitorEditor/Private/STimedDataMonitorPanel.h b/Engine/Plugins/VirtualProduction/TimedDataMonitor/Source/TimedDataMonitorEditor/Private/STimedDataMonitorPanel.h
--- a/Engine/Plugins/VirtualProduction/TimedDataMonitor/Source/TimedDataMonitorEditor/Private/STimedDataMonitorPanel.h
+++ b/Engine/Plugins/VirtualProduction/TimedDataMonitor/Source/TimedDataMonitorEditor/Private/STimedDataMonitorPanel.h
@@ -14,7 +14,6 @@ enum class ETimedDataMonitorTimeCorrectionReturnCode : uint8;
 
 struct FSlateBrush;
 struct FTimedDataMonitorCalibrationResult;
-class FTabManager;
 class FTimedDataMonitorCalibration;
 class FWorkspaceItem;
 class IMessageLogListing;
@@ -28,10 +27,9 @@ enum class ETimedDataMonitorEvaluationState : uint8;
 class STimedDataMonitorPanel : public SCompoundWidget
 {
 public:
-	static void RegisterNomadTabSpawner(TSharedPtr<FTabManager> TabManager, TSharedRef<FWorkspaceItem> InWorkspaceItem);
-	static void UnregisterNomadTabSpawner(TSharedPtr<FTabManager> TabManager);
+	static void RegisterNomadTabSpawner(TSharedRef<FWorkspaceItem> InWorkspaceItem);
+	static void UnregisterNomadTabSpawner();
 	static TSharedPtr<STimedDataMonitorPanel> GetPanelInstance();
-	static const FName TabName;
 
 private:
 	using Super = SCompoundWidget;
@@ -101,4 +99,6 @@ private:
 
 	bool bRefreshRequested = true;
 	double LastCachedValueUpdateTime = 0.0;
+
+	static FDelegateHandle LevelEditorTabManagerChangedHandle;
 };
diff --git a/Engine/Plugins/VirtualProduction/TimedDataMonitor/Source/TimedDataMonitorEditor/Private/TimedDataMonitorEditorModule.cpp b/Engine/Plugins/VirtualProduction/TimedDataMonitor/Source/TimedDataMonitorEditor/Private/TimedDataMonitorEditorModule.cpp
--- a/Engine/Plugins/VirtualProduction/TimedDataMonitor/Source/TimedDataMonitorEditor/Private/TimedDataMonitorEditorModule.cpp
+++ b/Engine/Plugins/VirtualProduction/TimedDataMonitor/Source/TimedDataMonitorEditor/Private/TimedDataMonitorEditorModule.cpp
@@ -1,96 +1,45 @@
 // Copyright Epic Games, Inc. All Rights Reserved.
 
-#include "TimedDataMonitorEditorModule.h"
 
-#include "Framework/Application/SlateApplication.h"
 #include "ISettingsModule.h"
-#include "LevelEditor.h"
 #include "STimedDataMonitorPanel.h"
-#include "STimedDataMonitorBufferVisualizer.h"
 #include "WorkspaceMenuStructureModule.h"
 #include "WorkspaceMenuStructure.h"
 
-#define LOCTEXT_NAMESPACE "TimedDataMonitorEditorModule"
 
-namespace UE::TDM
-{
-	static const FName NAME_LevelEditorModuleName("LevelEditor");
-}
+#define LOCTEXT_NAMESPACE "TimedDataMonitorEditorModule"
 
 
-void FTimedDataMonitorEditorModule::StartupModule()
+class FTimedDataMonitorEditorModule : public IModuleInterface
 {
-	if (ISettingsModule* SettingsModule = FModuleManager::GetModulePtr<ISettingsModule>("Settings"))
+public:
+	virtual void StartupModule() override
 	{
-		SettingsModule->RegisterSettings("Editor", "Plugins", "Timed Data Monitor",
-			LOCTEXT("SettingsName", "Timed Data Monitor"),
-			LOCTEXT("Description", "Configure the Timed Data Monitor panel."),
-			GetMutableDefault<UTimedDataMonitorEditorSettings>()
-		);
-	}
-
-	FLevelEditorModule& LevelEditorModule = FModuleManager::LoadModuleChecked<FLevelEditorModule>(UE::TDM::NAME_LevelEditorModuleName);
-	TSharedPtr<FTabManager> LevelEditorTabManager = LevelEditorModule.GetLevelEditorTabManager();
+		STimedDataMonitorPanel::RegisterNomadTabSpawner(WorkspaceMenu::GetMenuStructure().GetLevelEditorVirtualProductionCategory());
 
-	auto RegisterLambda = [this]()
+		if (ISettingsModule* SettingsModule = FModuleManager::GetModulePtr<ISettingsModule>("Settings"))
 		{
-			FLevelEditorModule& LevelEditorModule = FModuleManager::LoadModuleChecked<FLevelEditorModule>(UE::TDM::NAME_LevelEditorModuleName);
-			TSharedPtr<FTabManager> LevelEditorTabManager = LevelEditorModule.GetLevelEditorTabManager();
-			if (LevelEditorTabManager)
-			{
-				RegisterNomadTabSpawner(LevelEditorTabManager);
-			}
-		};
-
-	if (LevelEditorTabManager)
-	{
-		RegisterLambda();
-	}
-	else
-	{
-		LevelEditorTabManagerChangedHandle = LevelEditorModule.OnTabManagerChanged().AddLambda(RegisterLambda);
+			SettingsModule->RegisterSettings("Editor", "Plugins", "Timed Data Monitor",
+				LOCTEXT("SettingsName", "Timed Data Monitor"),
+				LOCTEXT("Description", "Configure the Timed Data Monitor panel."),
+				GetMutableDefault<UTimedDataMonitorEditorSettings>()
+			);
+		}
 	}
-}
 
-void FTimedDataMonitorEditorModule::ShutdownModule()
-{
-	if (!IsRunningCommandlet() && UObjectInitialized() && !IsEngineExitRequested())
+	virtual void ShutdownModule() override
 	{
-		ISettingsModule* SettingsModule = FModuleManager::GetModulePtr<ISettingsModule>("Settings");
-		if (SettingsModule != nullptr)
-		{
-			SettingsModule->UnregisterSettings("Editor", "Plugins", "Timed Data Monitor");
-		}
-
-		if (FSlateApplication::IsInitialized())
+		if (!IsRunningCommandlet() && UObjectInitialized() && !IsEngineExitRequested())
 		{
-			if (FLevelEditorModule* LevelEditorModule = FModuleManager::GetModulePtr<FLevelEditorModule>(UE::TDM::NAME_LevelEditorModuleName))
+			ISettingsModule* SettingsModule = FModuleManager::GetModulePtr<ISettingsModule>("Settings");
+			if (SettingsModule != nullptr)
 			{
-				LevelEditorModule->OnTabManagerChanged().Remove(LevelEditorTabManagerChangedHandle);
-
-				if (TSharedPtr<FTabManager> LevelEditorTabManager = LevelEditorModule->GetLevelEditorTabManager())
-				{
-					UnregisterNomadTabSpawner(LevelEditorTabManager);
-				}
+				SettingsModule->UnregisterSettings("Editor", "Plugins", "Timed Data Monitor");
 			}
+			STimedDataMonitorPanel::UnregisterNomadTabSpawner();
 		}
 	}
-}
-
-void FTimedDataMonitorEditorModule::RegisterNomadTabSpawner(TSharedPtr<FTabManager> TabManager)
-{
-	STimedDataMonitorPanel::RegisterNomadTabSpawner(TabManager, WorkspaceMenu::GetMenuStructure().GetLevelEditorVirtualProductionCategory());
-}
-
-void FTimedDataMonitorEditorModule::UnregisterNomadTabSpawner(TSharedPtr<FTabManager> TabManager)
-{
-	STimedDataMonitorPanel::UnregisterNomadTabSpawner(TabManager);
-}
-
-void FTimedDataMonitorEditorModule::DisplayTimedDataMonitorPanel(TSharedPtr<FTabManager> TabManager)
-{
-	TabManager->TryInvokeTab(STimedDataMonitorPanel::TabName);
-}
+};
 
 IMPLEMENT_MODULE(FTimedDataMonitorEditorModule, TimedDataMonitorEditor);
 
diff --git a/Engine/Plugins/VirtualProduction/TimedDataMonitor/Source/TimedDataMonitorEditor/Public/TimedDataMonitorEditorModule.h b/Engine/Plugins/VirtualProduction/TimedDataMonitor/Source/TimedDataMonitorEditor/Public/TimedDataMonitorEditorModule.h
deleted file mode 100644
diff --git a/Engine/Source/Programs/LiveLinkHub/Source/LiveLinkHub.Target.cs b/Engine/Source/Programs/LiveLinkHub/Source/LiveLinkHub.Target.cs
--- a/Engine/Source/Programs/LiveLinkHub/Source/LiveLinkHub.Target.cs
+++ b/Engine/Source/Programs/LiveLinkHub/Source/LiveLinkHub.Target.cs
@@ -65,9 +65,7 @@ public class LiveLinkHubTarget : TargetRules
 			"MediaIOFramework",
 			"MediaFrameworkUtilities",
 			"AjaMedia",
-			"BlackmagicMedia",
-
-			"TimedDataMonitor"
+			"BlackmagicMedia"
 		});
 
 		if (bEnableMetaHumanLiveLinkPlugin)
