Fix Horde issue 1174692

#rnx
#Synced-CL 49377257

diff --git a/Engine/Plugins/VirtualProduction/TimedDataMonitor/Source/TimedDataMonitorEditor/Private/STimingDiagramWidget.cpp b/Engine/Plugins/VirtualProduction/TimedDataMonitor/Source/TimedDataMonitorEditor/Private/STimingDiagramWidget.cpp
--- a/Engine/Plugins/VirtualProduction/TimedDataMonitor/Source/TimedDataMonitorEditor/Private/STimingDiagramWidget.cpp
+++ b/Engine/Plugins/VirtualProduction/TimedDataMonitor/Source/TimedDataMonitorEditor/Private/STimingDiagramWidget.cpp
@@ -3,6 +3,7 @@
 #include "STimingDiagramWidget.h"
 
 #include "Engine/Engine.h"
+#include "Styling/CoreStyle.h"
 #include "Styling/SlateTypes.h"
 #include "TimedDataMonitorEditorSettings.h"
 #include "TimedDataMonitorSubsystem.h"
