LiveLinkHub - Disable Timed Data Monitor temporarily to work around crash

#Synced-CL 49625772

diff --git a/Engine/Plugins/VirtualProduction/TimedDataMonitor/Source/TimedDataMonitor/Private/TimedDataMonitorSubsystem.cpp b/Engine/Plugins/VirtualProduction/TimedDataMonitor/Source/TimedDataMonitor/Private/TimedDataMonitorSubsystem.cpp
--- a/Engine/Plugins/VirtualProduction/TimedDataMonitor/Source/TimedDataMonitor/Private/TimedDataMonitorSubsystem.cpp
+++ b/Engine/Plugins/VirtualProduction/TimedDataMonitor/Source/TimedDataMonitor/Private/TimedDataMonitorSubsystem.cpp
@@ -8,6 +8,7 @@
 #include "ITimeManagementModule.h"
 #include "LatentActions.h"
 #include "Misc/App.h"
+#include "Misc/ConfigCacheIni.h"
 #include "Misc/CoreDelegates.h"
 #include "Misc/Paths.h"
 #include "TimedDataInputCollection.h"
@@ -100,29 +101,38 @@ void UTimedDataMonitorSubsystem::FTimeDataChannelItem::ResetValue()
 void UTimedDataMonitorSubsystem::Initialize(FSubsystemCollectionBase& Collection)
 {
 	Super::Initialize(Collection);
-	
-	bRequestSourceListRebuilt = true;
-	ITimeManagementModule::Get().GetTimedDataInputCollection().OnCollectionChanged().AddUObject(this, &UTimedDataMonitorSubsystem::OnTimedDataSourceCollectionChanged);
-	FCoreDelegates::OnEndFrame.AddUObject(this, &UTimedDataMonitorSubsystem::EndFrameCallback);
-	FCoreDelegates::OnBeginFrame.AddUObject(this, &UTimedDataMonitorSubsystem::BeginFrameCallback);
+
+	const bool bInitializeTimedDataMonitor = GConfig->GetBoolOrDefault(TEXT("TimedDataMonitor"), TEXT("bEnableTimedDataMonitor"), true, GEngineIni);
+	if (bInitializeTimedDataMonitor)
+	{
+		bRequestSourceListRebuilt = true;
+		ITimeManagementModule::Get().GetTimedDataInputCollection().OnCollectionChanged().AddUObject(this, &UTimedDataMonitorSubsystem::OnTimedDataSourceCollectionChanged);
+		FCoreDelegates::OnEndFrame.AddUObject(this, &UTimedDataMonitorSubsystem::EndFrameCallback);
+		FCoreDelegates::OnBeginFrame.AddUObject(this, &UTimedDataMonitorSubsystem::BeginFrameCallback);
+	}
 }
 
 
 void UTimedDataMonitorSubsystem::Deinitialize()
 {
-	FCoreDelegates::OnBeginFrame.RemoveAll(this);
-	FCoreDelegates::OnEndFrame.RemoveAll(this);
+	const bool bInitializeTimedDataMonitor = GConfig->GetBoolOrDefault(TEXT("TimedDataMonitor"), TEXT("bEnableTimedDataMonitor"), true, GEngineIni);
 
-	if (ITimeManagementModule::IsAvailable())
+	if (bInitializeTimedDataMonitor)
 	{
-		ITimeManagementModule::Get().GetTimedDataInputCollection().OnCollectionChanged().RemoveAll(this);
-	}
+		FCoreDelegates::OnBeginFrame.RemoveAll(this);
+		FCoreDelegates::OnEndFrame.RemoveAll(this);
 
-	bRequestSourceListRebuilt = true;
-	InputMap.Reset();
-	ChannelMap.Reset();
-	OnIdentifierListChanged_Delegate.Broadcast();
-	OnIdentifierListChanged_Dynamic.Broadcast();
+		if (ITimeManagementModule::IsAvailable())
+		{
+			ITimeManagementModule::Get().GetTimedDataInputCollection().OnCollectionChanged().RemoveAll(this);
+		}
+
+		bRequestSourceListRebuilt = true;
+		InputMap.Reset();
+		ChannelMap.Reset();
+		OnIdentifierListChanged_Delegate.Broadcast();
+		OnIdentifierListChanged_Dynamic.Broadcast();
+	}
 
 	Super::Deinitialize();
 }
diff --git a/Engine/Source/Programs/LiveLinkHub/Config/DefaultEngine.ini b/Engine/Source/Programs/LiveLinkHub/Config/DefaultEngine.ini
--- a/Engine/Source/Programs/LiveLinkHub/Config/DefaultEngine.ini
+++ b/Engine/Source/Programs/LiveLinkHub/Config/DefaultEngine.ini
@@ -13,6 +13,9 @@ bEnableReconnectingToStaleSource=true
 bEnableLLHMessageBusSourceFactory=true
 bUseMessagingConnectionManager=true
 
+[TimedDataMonitor]
+bEnableTimedDataMonitor=false
+
 [/Script/MacTargetPlatform.XcodeProjectSettings]
 TemplateMacPlist=(FilePath="/Engine/Source/Programs/LiveLinkHub/Resources/Mac/Info.plist")
 
