LiveLink and TDM: Option to save settings.

LL panel save button saves the last LL preset.
TDM panel save button saves all that it knows that affects the settings that TDM controls, which in practice are the LL preset, the media profile, and standalone media sources that appear in TDM.

Also fixed the frame delay and evaluation type of media not propagating all the way back to the media profile. Also auto-setting time synchronization flag in media source when timecode evaluation is picked in TDM, since it would otherwise have no effect.

#rnx

#Synced-CL 51080076

diff --git a/Engine/Plugins/Animation/LiveLink/Source/LiveLink/Private/LiveLinkClient.cpp b/Engine/Plugins/Animation/LiveLink/Source/LiveLink/Private/LiveLinkClient.cpp
--- a/Engine/Plugins/Animation/LiveLink/Source/LiveLink/Private/LiveLinkClient.cpp
+++ b/Engine/Plugins/Animation/LiveLink/Source/LiveLink/Private/LiveLinkClient.cpp
@@ -9,6 +9,7 @@
 #include "LiveLinkAnimationVirtualSubject.h"
 #include "LiveLinkLog.h"
 #include "LiveLinkMessages.h"
+#include "LiveLinkPreset.h"
 #include "LiveLinkPresetTypes.h"
 #include "LiveLinkRoleTrait.h"
 #include "LiveLinkProvider.h"
@@ -368,6 +369,16 @@ PRAGMA_ENABLE_DEPRECATION_WARNINGS
 #endif
 }
 
+void FLiveLinkClient::SetLastSavedPreset(ULiveLinkPreset* InPreset)
+{
+	LastSavedPreset = InPreset;
+}
+
+ULiveLinkPreset* FLiveLinkClient::GetLastSavedPreset() const
+{
+	return LastSavedPreset.Get();
+}
+
 void FLiveLinkClient::Shutdown()
 {
 	LLM_SCOPE_BYTAG(LiveLink_LiveLinkClient);
diff --git a/Engine/Plugins/Animation/LiveLink/Source/LiveLink/Private/LiveLinkTimedDataInput.cpp b/Engine/Plugins/Animation/LiveLink/Source/LiveLink/Private/LiveLinkTimedDataInput.cpp
--- a/Engine/Plugins/Animation/LiveLink/Source/LiveLink/Private/LiveLinkTimedDataInput.cpp
+++ b/Engine/Plugins/Animation/LiveLink/Source/LiveLink/Private/LiveLinkTimedDataInput.cpp
@@ -196,6 +196,11 @@ void FLiveLinkTimedDataInput::SetDataBufferSize(int32 BufferSize)
 	}
 }
 
+UObject* FLiveLinkTimedDataInput::GetSettingsObject() const
+{
+	return LiveLinkClient->GetSourceSettings(Source);
+}
+
 #if WITH_EDITOR
 const FSlateBrush* FLiveLinkTimedDataInput::GetDisplayIcon() const
 {
diff --git a/Engine/Plugins/Animation/LiveLink/Source/LiveLink/Public/LiveLinkClient.h b/Engine/Plugins/Animation/LiveLink/Source/LiveLink/Public/LiveLinkClient.h
--- a/Engine/Plugins/Animation/LiveLink/Source/LiveLink/Public/LiveLinkClient.h
+++ b/Engine/Plugins/Animation/LiveLink/Source/LiveLink/Public/LiveLinkClient.h
@@ -201,6 +201,9 @@ public:
 	/** (Internal use only) Cache subject settings for a given subject. Currently no-op outside of LiveLinkHub. */
 	virtual void CacheSubjectSettings(const FLiveLinkSubjectKey& SubjectKey, ULiveLinkSubjectSettings* Settings) const {}
 
+	UE_API void SetLastSavedPreset(ULiveLinkPreset* InPreset);
+	UE_API ULiveLinkPreset* GetLastSavedPreset() const;
+
 protected:
 	/**
 	 * Add delegates that will be triggered for all subjects.
@@ -346,6 +349,9 @@ private:
 	/** Whether or not parent subject support is enabled. Parent subjects allow resampling data to a different subject's rate before rebroadcasting it. */
 	bool bEnableParentSubjects = false;
 
+	/** Tracks the last preset saved or loaded by the user in this session. */
+	TWeakObjectPtr<ULiveLinkPreset> LastSavedPreset;
+
 #if WITH_EDITOR
 	/** Delegate when a subject is evaluated. */
 	FOnLiveLinkSubjectEvaluated OnLiveLinkSubjectEvaluatedDelegate;
diff --git a/Engine/Plugins/Animation/LiveLink/Source/LiveLink/Public/LiveLinkTimedDataInput.h b/Engine/Plugins/Animation/LiveLink/Source/LiveLink/Public/LiveLinkTimedDataInput.h
--- a/Engine/Plugins/Animation/LiveLink/Source/LiveLink/Public/LiveLinkTimedDataInput.h
+++ b/Engine/Plugins/Animation/LiveLink/Source/LiveLink/Public/LiveLinkTimedDataInput.h
@@ -33,6 +33,7 @@ public:
 	virtual bool IsDataBufferSizeControlledByInput() const override { return true; }
 	virtual void AddChannel(ITimedDataInputChannel* Channel) override { Channels.Add(Channel); }
 	virtual void RemoveChannel(ITimedDataInputChannel* Channel) override { Channels.RemoveSingleSwap(Channel); }
+	UE_API virtual UObject* GetSettingsObject() const override;
 #if WITH_EDITOR
 	UE_API virtual const FSlateBrush* GetDisplayIcon() const override;
 #endif
diff --git a/Engine/Plugins/Animation/LiveLink/Source/LiveLinkEditor/Private/LiveLinkClientPanelToolbar.cpp b/Engine/Plugins/Animation/LiveLink/Source/LiveLinkEditor/Private/LiveLinkClientPanelToolbar.cpp
--- a/Engine/Plugins/Animation/LiveLink/Source/LiveLinkEditor/Private/LiveLinkClientPanelToolbar.cpp
+++ b/Engine/Plugins/Animation/LiveLink/Source/LiveLinkEditor/Private/LiveLinkClientPanelToolbar.cpp
@@ -123,6 +123,26 @@ void SLiveLinkClientPanelToolbar::Construct(const FArguments& Args, FLiveLinkCli
 				]
 			]
 
+			// Save preset button
+			+ SHorizontalBox::Slot()
+			.Padding(8.f, 0.f, 0.f, 0.f)
+			.VAlign(VAlign_Center)
+			.HAlign(HAlign_Center)
+			.AutoWidth()
+			[
+				SNew(SButton)
+				.ToolTipText(LOCTEXT("SavePreset_Tip", "Save the current Live Link preset"))
+				.ForegroundColor(FSlateColor::UseForeground())
+				.ButtonStyle(FAppStyle::Get(), "HoverHintOnly")
+				.OnClicked(this, &SLiveLinkClientPanelToolbar::OnSavePreset)
+				.Visibility(Args._ShowPresetPicker ? EVisibility::Visible : EVisibility::Collapsed)
+				[
+					SNew(STextBlock)
+					.Font(FAppStyle::Get().GetFontStyle("FontAwesome.11"))
+					.Text(FEditorFontGlyphs::Floppy_O)
+				]
+			]
+
 			+ SHorizontalBox::Slot()
 			[
 				SNew(SSpacer)
@@ -308,7 +328,7 @@ void SLiveLinkClientPanelToolbar::OnSaveAsPreset()
 
 		FEditorFileUtils::PromptForCheckoutAndSave({ NewPackage }, false, false);
 	}
-	LiveLinkPreset = NewPreset;
+	Client->SetLastSavedPreset(NewPreset);
 }
 
 
@@ -334,7 +354,24 @@ void SLiveLinkClientPanelToolbar::OnImportPreset(const FAssetData& InPreset)
 		FScopedTransaction Transaction(LOCTEXT("ImportPreset_Transaction", "Import Live Link Preset"));
 		ImportedPreset->ApplyToClientLatent();
 	}
-	LiveLinkPreset = ImportedPreset;
+	Client->SetLastSavedPreset(ImportedPreset);
+}
+
+FReply SLiveLinkClientPanelToolbar::OnSavePreset()
+{
+	ULiveLinkPreset* CurrentPreset = Client->GetLastSavedPreset();
+	if (!CurrentPreset)
+	{
+		OnSaveAsPreset();
+		return FReply::Handled();
+	}
+
+	FScopedTransaction Transaction(LOCTEXT("SavePreset", "Save Preset"));
+	CurrentPreset->BuildFromClient();
+	CurrentPreset->MarkPackageDirty();
+	FEditorFileUtils::PromptForCheckoutAndSave({ CurrentPreset->GetPackage() }, false, false);
+
+	return FReply::Handled();
 }
 
 FReply SLiveLinkClientPanelToolbar::OnRevertChanges()
@@ -346,7 +383,7 @@ FReply SLiveLinkClientPanelToolbar::OnRevertChanges()
 	}
 
 	FScopedTransaction Transaction(LOCTEXT("RevertChanges_Transaction", "Revert Changes"));
-	ULiveLinkPreset* CurrentPreset = LiveLinkPreset.Get();
+	ULiveLinkPreset* CurrentPreset = Client->GetLastSavedPreset();
 	if (CurrentPreset)
 	{
 		CurrentPreset->ApplyToClientLatent();
@@ -357,7 +394,7 @@ FReply SLiveLinkClientPanelToolbar::OnRevertChanges()
 
 bool SLiveLinkClientPanelToolbar::HasLoadedLiveLinkPreset() const
 {
-	return LiveLinkPreset.IsValid();
+	return Client->GetLastSavedPreset() != nullptr;
 }
 
 
diff --git a/Engine/Plugins/Animation/LiveLink/Source/LiveLinkEditor/Public/LiveLinkClientPanelToolbar.h b/Engine/Plugins/Animation/LiveLink/Source/LiveLinkEditor/Public/LiveLinkClientPanelToolbar.h
--- a/Engine/Plugins/Animation/LiveLink/Source/LiveLinkEditor/Public/LiveLinkClientPanelToolbar.h
+++ b/Engine/Plugins/Animation/LiveLink/Source/LiveLinkEditor/Public/LiveLinkClientPanelToolbar.h
@@ -48,6 +48,8 @@ private:
 	TSharedRef<SWidget> OnPresetGeneratePresetsMenu();
 	/** Handles saving a preset to disk. */
 	void OnSaveAsPreset();
+	/** Saves the current preset, or falls back to Save As if none has been loaded. */
+	FReply OnSavePreset();
 	/** Handles importing a preset from disk. */
 	void OnImportPreset(const FAssetData& InPreset);
 	/** Handles reverting changes to a preset. */
@@ -58,8 +60,6 @@ private:
 private:
 	/** Cached livelink client. */
 	FLiveLinkClient* Client;
-	/** Last LiveLink preset that's been loaded. */
-	TWeakObjectPtr<ULiveLinkPreset> LiveLinkPreset;
 	/** Parent window override used for creating asset dialogs when running in LiveLinkHub. */
 	TSharedPtr<SWindow> ParentWindowOverride;
 };
diff --git a/Engine/Plugins/Media/MediaIOFramework/Source/MediaIOCore/Private/Player/MediaIOCorePlayerBase.cpp b/Engine/Plugins/Media/MediaIOFramework/Source/MediaIOCore/Private/Player/MediaIOCorePlayerBase.cpp
--- a/Engine/Plugins/Media/MediaIOFramework/Source/MediaIOCore/Private/Player/MediaIOCorePlayerBase.cpp
+++ b/Engine/Plugins/Media/MediaIOFramework/Source/MediaIOCore/Private/Player/MediaIOCorePlayerBase.cpp
@@ -218,6 +218,7 @@ bool FMediaIOCorePlayerBase::Open(const FString& Url, const IMediaOptions* Optio
 	ITimeManagementModule::Get().GetTimedDataInputCollection().Add(this);
 
 	OpenUrl = Url;
+	MediaSourceObject = const_cast<UObject*>(Options->ToUObject());
 	const bool bReadMediaOptions = ReadMediaOptions(Options);
 
 	// When AutoDetect is true, we can't create textures now as we don't know the video format.
@@ -1019,6 +1020,15 @@ void FMediaIOCorePlayerBase::SetEvaluationType(ETimedDataInputEvaluationType Eva
 	{
 		BaseSettings.EvaluationType = GetEvaluationType();
 		SetupSampleChannels();
+
+		if (UCaptureCardMediaSource* CaptureSource = Cast<UCaptureCardMediaSource>(MediaSourceObject.Get()))
+		{
+			CaptureSource->EvaluationType = EvaluationType;
+			if (EvaluationType == EMediaIOSampleEvaluationType::Timecode)
+			{
+				CaptureSource->bUseTimeSynchronization = true;
+			}
+		}
 	}
 }
 
@@ -1054,7 +1064,6 @@ void FMediaIOCorePlayerBase::SetEvaluationOffsetInSeconds(double Offset)
 
 	case EMediaIOSampleEvaluationType::Timecode:
 	{
-		//Media doesn't support subframes playback (interpolation between frames) so offsets are always in full frame
 		const double FrameOffset = ITimedDataInput::ConvertSecondOffsetInFrameOffset(Offset, VideoFrameRate);
 		FrameDelay = FMath::CeilToInt(FrameOffset);
 		break;
@@ -1064,6 +1073,12 @@ void FMediaIOCorePlayerBase::SetEvaluationOffsetInSeconds(double Offset)
 		checkNoEntry();
 		break;
 	}
+
+	if (UTimeSynchronizableMediaSource* Source = Cast<UTimeSynchronizableMediaSource>(MediaSourceObject.Get()))
+	{
+		Source->FrameDelay = FrameDelay;
+		Source->TimeDelay = TimeDelay;
+	}
 }
 
 FFrameRate FMediaIOCorePlayerBase::GetFrameRate() const
@@ -1092,6 +1107,11 @@ bool FMediaIOCorePlayerBase::SupportsSubFrames() const
 	return false;
 }
 
+UObject* FMediaIOCorePlayerBase::GetSettingsObject() const
+{
+	return MediaSourceObject.Get();
+}
+
 void FMediaIOCorePlayerBase::AddAudioSample(const TSharedRef<FMediaIOCoreAudioSampleBase>& Sample)
 {
 	GetSamples_Internal().AddAudio(Sample);
diff --git a/Engine/Plugins/Media/MediaIOFramework/Source/MediaIOCore/Public/MediaIOCorePlayerBase.h b/Engine/Plugins/Media/MediaIOFramework/Source/MediaIOCore/Public/MediaIOCorePlayerBase.h
--- a/Engine/Plugins/Media/MediaIOFramework/Source/MediaIOCore/Public/MediaIOCorePlayerBase.h
+++ b/Engine/Plugins/Media/MediaIOFramework/Source/MediaIOCore/Public/MediaIOCorePlayerBase.h
@@ -164,6 +164,7 @@ public:
 	UE_API virtual void AddChannel(ITimedDataInputChannel* Channel) override;
 	UE_API virtual void RemoveChannel(ITimedDataInputChannel * Channel) override;
 	UE_API virtual bool SupportsSubFrames() const override;
+	UE_API virtual UObject* GetSettingsObject() const override;
 
 public:
 	/** Class used to pass information about current frame for sample picking.*/
@@ -330,6 +331,9 @@ protected:
 	/** Open color IO conversion data. */
 	TSharedPtr<struct FOpenColorIOColorConversionSettings> OCIOSettings;
 
+	/** The media source that opened this player, for settings persistence. */
+	TWeakObjectPtr<UObject> MediaSourceObject;
+
 	/** Whether media playback should be framelocked to the engine's timecode */
 	bool bFramelock = false;
 
diff --git a/Engine/Plugins/VirtualProduction/TimedDataMonitor/Source/TimedDataMonitorEditor/Private/STimedDataListView.cpp b/Engine/Plugins/VirtualProduction/TimedDataMonitor/Source/TimedDataMonitorEditor/Private/STimedDataListView.cpp
--- a/Engine/Plugins/VirtualProduction/TimedDataMonitor/Source/TimedDataMonitorEditor/Private/STimedDataListView.cpp
+++ b/Engine/Plugins/VirtualProduction/TimedDataMonitor/Source/TimedDataMonitorEditor/Private/STimedDataListView.cpp
@@ -493,6 +493,7 @@ void STimedDataInputTableRow::OnEnabledCheckStateChanged(ECheckBoxState NewState
 	{
 		TimedDataMonitorSubsystem->SetChannelEnabled(Item->ChannelIdentifier, NewState == ECheckBoxState::Checked);
 	}
+	OwnerTreeView->MarkSettingsDirty();
 	OwnerTreeView->RequestRefresh();
 }
 
@@ -572,7 +573,8 @@ void STimedDataInputTableRow::SetEvaluationOffset(double InValue, ETextCommit::T
 		{
 			ensure(false);
 		}
-		
+
+		OwnerTreeView->MarkSettingsDirty();
 		OwnerTreeView->RequestRefresh();
 	}
 }
@@ -619,6 +621,7 @@ void STimedDataInputTableRow::SetBufferSize(int32 InValue, ETextCommit::Type InT
 			TimedDataMonitorSubsystem->SetChannelDataBufferSize(Item->ChannelIdentifier, InValue);
 			Item->CachedBufferSize = TimedDataMonitorSubsystem->GetChannelDataBufferSize(Item->ChannelIdentifier);
 		}
+		OwnerTreeView->MarkSettingsDirty();
 		OwnerTreeView->RequestRefresh();
 	}
 }
@@ -701,6 +704,7 @@ void STimedDataInputTableRow::SetInputEvaluationType(ETimedDataInputEvaluationTy
 	{
 		TimedDataMonitorSubsystem->SetInputEvaluationType(Item->InputIdentifier, EvaluationType);
 
+		OwnerTreeView->MarkSettingsDirty();
 		OwnerTreeView->RequestRefresh();
 	}
 }
@@ -845,6 +849,15 @@ void STimedDataInputListView::RequestRefresh()
 }
 
 
+void STimedDataInputListView::MarkSettingsDirty()
+{
+	if (TSharedPtr<STimedDataMonitorPanel> OwnerPanelPin = OwnerPanel.Pin())
+	{
+		OwnerPanelPin->MarkSettingsDirty();
+	}
+}
+
+
 void STimedDataInputListView::UpdateCachedValue()
 {
 	if (bRebuildListRequested)
@@ -947,6 +960,11 @@ void STimedDataInputListView::OnToggleAllEnabledCheckState(ECheckBoxState CheckB
 	{
 		TimedDataMonitorSubsystem->SetInputEnabled(RowDataPtr->InputIdentifier, bIsEnabled);
 	}
+
+	if (TSharedPtr<STimedDataMonitorPanel> OwnerPanelPin = OwnerPanel.Pin())
+	{
+		OwnerPanelPin->MarkSettingsDirty();
+	}
 }
 
 
diff --git a/Engine/Plugins/VirtualProduction/TimedDataMonitor/Source/TimedDataMonitorEditor/Private/STimedDataListView.h b/Engine/Plugins/VirtualProduction/TimedDataMonitor/Source/TimedDataMonitorEditor/Private/STimedDataListView.h
--- a/Engine/Plugins/VirtualProduction/TimedDataMonitor/Source/TimedDataMonitorEditor/Private/STimedDataListView.h
+++ b/Engine/Plugins/VirtualProduction/TimedDataMonitor/Source/TimedDataMonitorEditor/Private/STimedDataListView.h
@@ -106,6 +106,7 @@ public:
 
 	void RequestRefresh();
 	void UpdateCachedValue();
+	void MarkSettingsDirty();
 
 	FTimedDataMonitorInputIdentifier GetSelectedInputIdentifier() const;
 
diff --git a/Engine/Plugins/VirtualProduction/TimedDataMonitor/Source/TimedDataMonitorEditor/Private/STimedDataMonitorPanel.cpp b/Engine/Plugins/VirtualProduction/TimedDataMonitor/Source/TimedDataMonitorEditor/Private/STimedDataMonitorPanel.cpp
--- a/Engine/Plugins/VirtualProduction/TimedDataMonitor/Source/TimedDataMonitorEditor/Private/STimedDataMonitorPanel.cpp
+++ b/Engine/Plugins/VirtualProduction/TimedDataMonitor/Source/TimedDataMonitorEditor/Private/STimedDataMonitorPanel.cpp
@@ -11,17 +11,32 @@
 #include "Misc/App.h"
 #include "TimedDataMonitorSubsystem.h"
 
+#include "AssetRegistry/AssetRegistryModule.h"
+#include "AssetToolsModule.h"
+#include "ContentBrowserModule.h"
 #include "EditorFontGlyphs.h"
+#include "Features/IModularFeatures.h"
+#include "FileHelpers.h"
+#include "IAssetTools.h"
+#include "IContentBrowserSingleton.h"
+#include "ILiveLinkClient.h"
+#include "LiveLinkClient.h"
+#include "LiveLinkPreset.h"
+#include "LiveLinkSettings.h"
+#include "Misc/PackageName.h"
 #include "Modules/ModuleManager.h"
+#include "Profile/IMediaProfileManager.h"
+#include "Profile/MediaProfile.h"
+#include "ScopedTransaction.h"
 #include "TimedDataMonitorEditorStyle.h"
 
 #include "Framework/MultiBox/MultiBoxBuilder.h"
 
+#include "SlateOptMacros.h"
 #include "STimedDataGenlock.h"
+#include "STimedDataListView.h"
 #include "STimedDataMonitorBufferVisualizer.h"
 #include "STimedDataTimecodeProvider.h"
-#include "STimedDataListView.h"
-#include "SlateOptMacros.h"
 #include "Widgets/Docking/SDockTab.h"
 #include "Widgets/Images/SThrobber.h"
 #include "Widgets/Input/SButton.h"
@@ -239,6 +254,25 @@ void STimedDataMonitorPanel::Construct(const FArguments& InArgs)
 					.OnGetMenuContent(this, &STimedDataMonitorPanel::OnResetBuildMenu)
 				]
 			]
+			// Save preset button
+			+ SHorizontalBox::Slot()
+			.AutoWidth()
+			.VAlign(VAlign_Center)
+			.Padding(4.f)
+			[
+				SNew(SButton)
+				.ButtonStyle(FTimedDataMonitorEditorStyle::Get(), "ToggleButton")
+				.ToolTipText(LOCTEXT("SavePreset_ToolTip", "Save all monitored input settings and the Live Link preset"))
+				.ContentPadding(FMargin(4, 4))
+				.HAlign(HAlign_Center)
+				.VAlign(VAlign_Center)
+				.OnClicked(this, &STimedDataMonitorPanel::OnSavePreset)
+				[
+					SNew(SImage)
+					.Image(this, &STimedDataMonitorPanel::GetSaveButtonBrush)
+					.ColorAndOpacity(FSlateColor::UseForeground())
+				]
+			]
 			// show buffers
 			+ SHorizontalBox::Slot()
 			.AutoWidth()
@@ -641,6 +675,7 @@ void STimedDataMonitorPanel::OnResetAllEvaluationTimeClicked()
 	{
 		TimedDataMonitorSubsystem->SetInputEvaluationOffsetInSeconds(Input, 0.f);
 	}
+	MarkSettingsDirty();
 }
 
 
@@ -729,6 +764,7 @@ void STimedDataMonitorPanel::ApplyTimeCorrectionOnSelection()
 		if (SelectedInput.IsValid())
 		{
 			ApplyTimeCorrection(SelectedInput);
+			MarkSettingsDirty();
 		}
 	}
 }
@@ -743,6 +779,7 @@ void STimedDataMonitorPanel::ResetTimeCorrectionOnSelection()
 		{
 			UTimedDataMonitorSubsystem* TimedDataMonitorSubsystem = GEngine->GetEngineSubsystem<UTimedDataMonitorSubsystem>();
 			TimedDataMonitorSubsystem->SetInputEvaluationOffsetInSeconds(SelectedInput, 0.f);
+			MarkSettingsDirty();
 		}
 	}
 }
@@ -795,6 +832,7 @@ void STimedDataMonitorPanel::CalibrateWithTimecodeCompleted(FTimedDataMonitorCal
 {
 	UTimedDataMonitorSubsystem* TimedDataMonitorSubsystem = GEngine->GetEngineSubsystem<UTimedDataMonitorSubsystem>();
 	MonitorCalibration.Reset();
+	MarkSettingsDirty();
 
 	if (Result.ReturnCode != ETimedDataMonitorCalibrationReturnCode::Succeeded)
 	{
@@ -885,6 +923,8 @@ void STimedDataMonitorPanel::ApplyTimeCorrectionAll()
 			, LOCTEXT("JamSucceeded", "Time Correction succeeded.")));
 	}
 
+	MarkSettingsDirty();
+
 	GetMutableDefault<UTimedDataMonitorEditorSettings>()->LastCalibrationType = ETimedDataMonitorEditorCalibrationType::CalibrateWithTimecode;
 	GetMutableDefault<UTimedDataMonitorEditorSettings>()->SaveConfig();
 }
@@ -982,4 +1022,216 @@ FReply STimedDataMonitorPanel::OnCancelCalibration()
 	return FReply::Handled();
 }
 
+FLiveLinkClient* STimedDataMonitorPanel::GetLiveLinkClient() const
+{
+	if (IModularFeatures::Get().IsModularFeatureAvailable(ILiveLinkClient::ModularFeatureName))
+	{
+		return static_cast<FLiveLinkClient*>(&IModularFeatures::Get().GetModularFeature<ILiveLinkClient>(ILiveLinkClient::ModularFeatureName));
+	}
+	return nullptr;
+}
+
+namespace SavePresetUtils
+{
+	static bool OpenSaveDialog(const FString& InDefaultPath, const FString& InNewNameSuggestion, FString& OutPackageName)
+	{
+		FSaveAssetDialogConfig SaveAssetDialogConfig;
+		SaveAssetDialogConfig.DefaultPath = InDefaultPath;
+		SaveAssetDialogConfig.DefaultAssetName = InNewNameSuggestion;
+		SaveAssetDialogConfig.AssetClassNames.Add(ULiveLinkPreset::StaticClass()->GetClassPathName());
+		SaveAssetDialogConfig.ExistingAssetPolicy = ESaveAssetDialogExistingAssetPolicy::AllowButWarn;
+		SaveAssetDialogConfig.DialogTitleOverride = LOCTEXT("SaveLiveLinkPresetDialogTitle", "Save Live Link Preset");
+
+		FContentBrowserModule& ContentBrowserModule = FModuleManager::LoadModuleChecked<FContentBrowserModule>("ContentBrowser");
+		FString SaveObjectPath = ContentBrowserModule.Get().CreateModalSaveAssetDialog(SaveAssetDialogConfig);
+
+		if (!SaveObjectPath.IsEmpty())
+		{
+			OutPackageName = FPackageName::ObjectPathToPackageName(SaveObjectPath);
+			return true;
+		}
+		return false;
+	}
+
+	static bool GetSavePresetPackageName(FString& OutName)
+	{
+		ULiveLinkUserSettings* ConfigSettings = GetMutableDefault<ULiveLinkUserSettings>();
+
+		FDateTime Today = FDateTime::Now();
+		TMap<FString, FStringFormatArg> FormatArgs;
+		FormatArgs.Add(TEXT("date"), Today.ToString(TEXT("%Y-%m-%d-%H%M%S")));
+
+		const FString DefaultSaveDirectory = FString::Format(*ConfigSettings->GetPresetSaveDir().Path, FormatArgs);
+
+		FString DialogStartPath;
+
+		FPackageName::TryConvertFilenameToLongPackageName(DefaultSaveDirectory, DialogStartPath);
+
+		if (DialogStartPath.IsEmpty())
+		{
+			DialogStartPath = TEXT("/Game");
+		}
+
+		FString DefaultName = LOCTEXT("NewLiveLinkPreset", "NewLiveLinkPreset").ToString();
+
+		FString UniquePackageName;
+		FString UniqueAssetName;
+		FAssetToolsModule& AssetToolsModule = FModuleManager::LoadModuleChecked<FAssetToolsModule>("AssetTools");
+
+		AssetToolsModule.Get().CreateUniqueAssetName(DialogStartPath / DefaultName, TEXT(""), UniquePackageName, UniqueAssetName);
+
+		FString DialogStartName = FPaths::GetCleanFilename(UniqueAssetName);
+		FString UserPackageName;
+		FString NewPackageName;
+
+		bool bFilenameValid = false;
+
+		while (!bFilenameValid)
+		{
+			if (!OpenSaveDialog(DialogStartPath, DialogStartName, UserPackageName))
+			{
+				return false;
+			}
+
+			NewPackageName = FString::Format(*UserPackageName, FormatArgs);
+
+			FText OutError;
+			bFilenameValid = FPackageName::IsValidLongPackageName(NewPackageName, false, &OutError);
+		}
+
+		ConfigSettings->PresetSaveDir.Path = FPackageName::GetLongPackagePath(UserPackageName);
+		ConfigSettings->SaveConfig();
+
+		OutName = MoveTemp(NewPackageName);
+
+		return true;
+	}
+}
+
+void STimedDataMonitorPanel::CollectInputSettingsPackages(TArray<UPackage*>& OutPackages)
+{
+	UMediaProfile* MediaProfile = IMediaProfileManager::Get().GetCurrentMediaProfile();
+
+	if (MediaProfile)
+	{
+		MediaProfile->MarkPackageDirty();
+		OutPackages.AddUnique(MediaProfile->GetPackage());
+	}
+
+	UTimedDataMonitorSubsystem* Subsystem = GEngine->GetEngineSubsystem<UTimedDataMonitorSubsystem>();
+
+	if (!Subsystem)
+	{
+		return;
+	}
+
+	TArray<FTimedDataMonitorInputIdentifier> AllInputs = Subsystem->GetAllInputs();
+
+	for (const FTimedDataMonitorInputIdentifier& InputId : AllInputs)
+	{
+		ITimedDataInput* Input = Subsystem->GetTimedDataInput(InputId);
+
+		if (!Input)
+		{
+			continue;
+		}
+
+		UObject* SettingsObject = Input->GetSettingsObject();
+
+		if (!SettingsObject)
+		{
+			continue;
+		}
+
+		UPackage* Package = SettingsObject->GetPackage();
+
+		if (Package && Package != GetTransientPackage())
+		{
+			SettingsObject->MarkPackageDirty();
+			OutPackages.AddUnique(Package);
+		}
+	}
+}
+
+const FSlateBrush* STimedDataMonitorPanel::GetSaveButtonBrush() const
+{
+	return FAppStyle::GetBrush(bSettingsDirty ? TEXT("Icons.SaveChanged") : TEXT("Icons.Save"));
+}
+
+FReply STimedDataMonitorPanel::OnSavePreset()
+{
+	FLiveLinkClient* LiveLinkClient = GetLiveLinkClient();
+
+	if (!LiveLinkClient)
+	{
+		return FReply::Handled();
+	}
+
+	ULiveLinkPreset* CurrentPreset = LiveLinkClient->GetLastSavedPreset();
+
+	if (!CurrentPreset)
+	{
+		OnSaveAsPreset();
+		return FReply::Handled();
+	}
+
+	TArray<UPackage*> PackagesToSave;
+
+	FScopedTransaction Transaction(LOCTEXT("SavePreset", "Save Preset"));
+
+	CurrentPreset->BuildFromClient();
+	CurrentPreset->MarkPackageDirty();
+	PackagesToSave.Add(CurrentPreset->GetPackage());
+
+	CollectInputSettingsPackages(PackagesToSave);
+
+	const FEditorFileUtils::EPromptReturnCode Result = FEditorFileUtils::PromptForCheckoutAndSave(PackagesToSave, false, false);
+
+	if (Result == FEditorFileUtils::PR_Success)
+	{
+		bSettingsDirty = false;
+	}
+
+	return FReply::Handled();
+}
+
+void STimedDataMonitorPanel::OnSaveAsPreset()
+{
+	FString PackageName;
+	if (!SavePresetUtils::GetSavePresetPackageName(PackageName))
+	{
+		return;
+	}
+
+	FScopedTransaction Transaction(LOCTEXT("SaveAsPreset", "Save As Preset"));
+
+	const FString NewAssetName = FPackageName::GetLongPackageAssetName(PackageName);
+	UPackage* NewPackage = CreatePackage(*PackageName);
+	ULiveLinkPreset* NewPreset = NewObject<ULiveLinkPreset>(NewPackage, *NewAssetName, RF_Public | RF_Standalone | RF_Transactional);
+
+	if (NewPreset)
+	{
+		TArray<UPackage*> PackagesToSave;
+
+		NewPreset->BuildFromClient();
+		NewPreset->MarkPackageDirty();
+		FAssetRegistryModule::AssetCreated(NewPreset);
+		PackagesToSave.Add(NewPackage);
+
+		CollectInputSettingsPackages(PackagesToSave);
+
+		const FEditorFileUtils::EPromptReturnCode Result = FEditorFileUtils::PromptForCheckoutAndSave(PackagesToSave, false, false);
+		if (Result == FEditorFileUtils::PR_Success)
+		{
+			bSettingsDirty = false;
+
+			FLiveLinkClient* LiveLinkClient = GetLiveLinkClient();
+			if (LiveLinkClient)
+			{
+				LiveLinkClient->SetLastSavedPreset(NewPreset);
+			}
+		}
+	}
+}
+
 #undef LOCTEXT_NAMESPACE
diff --git a/Engine/Plugins/VirtualProduction/TimedDataMonitor/Source/TimedDataMonitorEditor/Private/STimedDataMonitorPanel.h b/Engine/Plugins/VirtualProduction/TimedDataMonitor/Source/TimedDataMonitorEditor/Private/STimedDataMonitorPanel.h
--- a/Engine/Plugins/VirtualProduction/TimedDataMonitor/Source/TimedDataMonitorEditor/Private/STimedDataMonitorPanel.h
+++ b/Engine/Plugins/VirtualProduction/TimedDataMonitor/Source/TimedDataMonitorEditor/Private/STimedDataMonitorPanel.h
@@ -11,9 +11,10 @@
 
 enum class ETimedDataMonitorTimeCorrectionReturnCode : uint8;
 
-
 struct FSlateBrush;
 struct FTimedDataMonitorCalibrationResult;
+
+class FLiveLinkClient;
 class FTabManager;
 class FTimedDataMonitorCalibration;
 class FWorkspaceItem;
@@ -22,6 +23,7 @@ class STimedDataGenlock;
 class STimedDataInputListView;
 class STimedDataTimecodeProvider;
 class SWidget;
+class ULiveLinkPreset;
 
 enum class ETimedDataMonitorEvaluationState : uint8;
 
@@ -47,6 +49,12 @@ public:
 
 	void RequestRefresh() { bRefreshRequested = true; }
 
+	/** Flags that monitored input settings have been modified and need saving. */
+	void MarkSettingsDirty()
+	{
+		bSettingsDirty = true;
+	}
+
 private:
 	FReply OnCalibrateClicked();
 	TSharedRef<SWidget> OnCalibrateBuildMenu();
@@ -83,6 +91,17 @@ private:
 	ETimedDataMonitorTimeCorrectionReturnCode ApplyTimeCorrection(const FTimedDataMonitorInputIdentifier& InputIndentifier);
 	FReply OnCancelCalibration();
 
+	/** Saves the tracked preset and modified input settings, or falls back to Save As if no preset is tracked. */
+	FReply OnSavePreset();
+	/** Opens a Save As dialog to create a new Live Link preset and saves all input settings. */
+	void OnSaveAsPreset();
+	/** Collects packages from all monitored inputs that have a settings object (e.g. media sources). */
+	void CollectInputSettingsPackages(TArray<UPackage*>& OutPackages);
+	/** Returns Icons.SaveChanged when dirty, Icons.Save otherwise. */
+	const FSlateBrush* GetSaveButtonBrush() const;
+	/** Returns the Live Link client from modular features, or nullptr. */
+	FLiveLinkClient* GetLiveLinkClient() const;
+
 private:
 	TSharedPtr<STimedDataGenlock> TimedDataGenlockWidget;
 	TSharedPtr<STimedDataTimecodeProvider> TimedDataTimecodeWidget;
@@ -100,5 +119,7 @@ private:
 	TUniquePtr<FTimedDataMonitorCalibration> MonitorCalibration;
 
 	bool bRefreshRequested = true;
+	/** True when monitored input settings have been modified since last save. */
+	bool bSettingsDirty = false;
 	double LastCachedValueUpdateTime = 0.0;
 };
diff --git a/Engine/Plugins/VirtualProduction/TimedDataMonitor/Source/TimedDataMonitorEditor/TimedDataMonitorEditor.build.cs b/Engine/Plugins/VirtualProduction/TimedDataMonitor/Source/TimedDataMonitorEditor/TimedDataMonitorEditor.build.cs
--- a/Engine/Plugins/VirtualProduction/TimedDataMonitor/Source/TimedDataMonitorEditor/TimedDataMonitorEditor.build.cs
+++ b/Engine/Plugins/VirtualProduction/TimedDataMonitor/Source/TimedDataMonitorEditor/TimedDataMonitorEditor.build.cs
@@ -8,6 +8,9 @@ public class TimedDataMonitorEditor : ModuleRules
 	{
 		PrivateDependencyModuleNames.AddRange(
 			new string[] {
+				"AssetRegistry",
+				"AssetTools",
+				"ContentBrowser",
 				"Core",
 				"CoreUObject",
 				"EditorFramework",
@@ -15,6 +18,9 @@ public class TimedDataMonitorEditor : ModuleRules
 				"Engine",
 				"EditorStyle",
 				"InputCore",
+				"LiveLink",
+				"LiveLinkInterface",
+				"MediaFrameworkUtilities",
 				"MessageLog",
 				"Projects",
 				"Settings",
diff --git a/Engine/Plugins/VirtualProduction/TimedDataMonitor/TimedDataMonitor.uplugin b/Engine/Plugins/VirtualProduction/TimedDataMonitor/TimedDataMonitor.uplugin
--- a/Engine/Plugins/VirtualProduction/TimedDataMonitor/TimedDataMonitor.uplugin
+++ b/Engine/Plugins/VirtualProduction/TimedDataMonitor/TimedDataMonitor.uplugin
@@ -25,5 +25,15 @@
 			"Type": "Editor",
 			"LoadingPhase": "Default"
 		}
+	],
+	"Plugins": [
+		{
+			"Name": "LiveLink",
+			"Enabled": true
+		},
+		{
+			"Name": "MediaFrameworkUtilities",
+			"Enabled": true
+		}
 	]
 }
\ No newline at end of file
diff --git a/Engine/Source/Runtime/TimeManagement/Public/ITimedDataInput.h b/Engine/Source/Runtime/TimeManagement/Public/ITimedDataInput.h
--- a/Engine/Source/Runtime/TimeManagement/Public/ITimedDataInput.h
+++ b/Engine/Source/Runtime/TimeManagement/Public/ITimedDataInput.h
@@ -17,6 +17,7 @@
 class ITimedDataInput;
 class ITimedDataInputChannel;
 class SWidget;
+class UObject;
 struct FSlateBrush;
 
 
@@ -142,6 +143,12 @@ public:
 	/** Convert frame offset to second offset using this input's framerate. */
 	TIMEMANAGEMENT_API double ConvertFrameOffsetInSecondOffset(double Frames) const;
 
+	/** Get the UObject that owns the settings for this input, for saving purposes. */
+	virtual UObject* GetSettingsObject() const
+	{
+		return nullptr;
+	}
+
 #if WITH_EDITOR
 	/** Get the icon that represent the input. */
 	virtual const FSlateBrush* GetDisplayIcon() const = 0;
